Media Viewer: Initial Commit

#virtualized

diff --git a/Engine/Plugins/Experimental/MediaViewer/Source/MediaViewer/Private/ImageViewers/MediaImageViewer.cpp b/Engine/Plugins/Experimental/MediaViewer/Source/MediaViewer/Private/ImageViewers/MediaImageViewer.cpp
new file mode 100644
--- /dev/null
+++ b/Engine/Plugins/Experimental/MediaViewer/Source/MediaViewer/Private/ImageViewers/MediaImageViewer.cpp
@@ -0,0 +1,172 @@
+// Copyright Epic Games, Inc. All Rights Reserved.
+
+#include "ImageViewer/MediaImageViewer.h"
+
+#include "Engine/Texture.h"
+#include "Math/Vector2D.h"
+#include "Misc/TVariant.h"
+#include "Widgets/SMediaViewer.h"
+
+namespace UE::MediaViewer
+{
+
+const FSlateColorBrush FMediaImageViewer::BackgroundColorBrush(FLinearColor::White);
+
+FMediaImageViewer::FMediaImageViewer(const FMediaImageViewerInfo& InImageInfo)
+	: ImageInfo(InImageInfo)
+	, PanelSettings(MakeShared<TStructOnScope<FMediaImageViewerPanelSettings>>())
+	, PaintSettings(MakeShared<TStructOnScope<FMediaImagePaintSettings>>())
+	, BackgroundImageBrush(static_cast<UObject*>(nullptr), FVector2D(1))
+	, Brush(nullptr)
+{
+	PanelSettings->InitializeAs<FMediaImageViewerPanelSettings>();
+	PaintSettings->InitializeAs<FMediaImagePaintSettings>();
+}
+
+FMediaImageViewer::~FMediaImageViewer()
+{
+	static_cast<FStructOnScope&>(PanelSettings.Get()).Destroy();
+	static_cast<FStructOnScope&>(PaintSettings.Get()).Destroy();
+}
+
+void FMediaImageViewer::NotifyPostChange(const FPropertyChangedEvent& InPropertyChangedEvent, FProperty* InPropertyThatChanged)
+{
+	if (UTexture* Texture = GetPanelSettings().BackgroundTexture.LoadSynchronous())
+	{
+		BackgroundImageBrush.SetResourceObject(Texture);
+		BackgroundImageBrush.ImageSize.X = Texture->GetSurfaceWidth();
+		BackgroundImageBrush.ImageSize.Y = Texture->GetSurfaceHeight();
+	}
+	else
+	{
+		BackgroundImageBrush.SetResourceObject(nullptr);
+	}
+}
+
+FVector2D FMediaImageViewer::GetViewerCenter(const FVector2D& InViewerSize) const
+{
+	return InViewerSize * 0.5;
+}
+
+FVector2D FMediaImageViewer::GetPaintOffsetForViewerCenter(const FVector2D& InViewerSize) const
+{
+	return GetViewerCenter(InViewerSize) - (FVector2D(ImageInfo.Size.X, ImageInfo.Size.Y) * 0.5 * GetPaintSettings().Scale);
+}
+
+FVector2D FMediaImageViewer::GetPaintOffset(const FVector2D& InViewerSize, const FVector2D& InViewerPosition) const
+{
+	return GetPaintOffsetForViewerCenter(InViewerSize) + (FVector2D(GetPaintSettings().Offset)) + InViewerPosition;
+}
+
+FVector2D FMediaImageViewer::GetPaintSize() const
+{
+	return FVector2D(ImageInfo.Size.X, ImageInfo.Size.Y) * GetPaintSettings().Scale;
+}
+
+FSlateClippingZone FMediaImageViewer::CreateSlateClippingZone(const FGeometry& InGeometry, const FVector2D& InViewerSize, 
+	const FVector2D& InViewerPosition, EOrientation InOrientation, const FFloatRange& InUVRange) const
+{
+	FSlateClippingZone Zone(InGeometry);
+
+	if (InOrientation == Orient_Horizontal)
+	{
+		const float Offset = InUVRange.GetLowerBound().GetValue() * InViewerSize.X;
+		const float Width = InUVRange.Size<float>() * InViewerSize.X;
+
+		Zone.TopLeft.X += Offset + InViewerPosition.X;
+		Zone.TopRight.X = Zone.TopLeft.X + Width;
+
+		Zone.BottomLeft.X += Offset + InViewerPosition.X;
+		Zone.BottomRight.X = Zone.BottomLeft.X + Width;
+	}
+	else
+	{
+		const float Offset = InUVRange.GetLowerBound().GetValue() * InViewerSize.Y;
+		const float Height = InUVRange.Size<float>() * InViewerSize.Y;
+
+		Zone.TopLeft.Y += Offset + InViewerPosition.Y;
+		Zone.BottomLeft.Y = Zone.TopLeft.Y + Height;
+
+		Zone.TopRight.Y += Offset + InViewerPosition.Y;
+		Zone.BottomRight.Y = Zone.TopRight.Y + Height;
+	}
+
+	return Zone;
+}
+
+void FMediaImageViewer::Paint(FMediaImagePaintParams& InPaintParams)
+{
+	FVector2D PaintOffset = GetPaintOffset(InPaintParams.ViewerSize, InPaintParams.ViewerPosition);
+	FVector2D PaintSize = GetPaintSize();
+	FPaintGeometry PaintGeometry = InPaintParams.AllottedGeometry.ToPaintGeometry(PaintSize, FSlateLayoutTransform(PaintOffset));
+
+	const FMediaImagePaintGeometry Geometry = {
+		MoveTemp(PaintOffset),
+		MoveTemp(PaintSize),
+		MoveTemp(PaintGeometry)
+	};
+
+	const FSlateClippingZone ClippingZone = CreateSlateClippingZone(
+		InPaintParams.AllottedGeometry,
+		InPaintParams.ViewerSize,
+		InPaintParams.ViewerPosition,
+		InPaintParams.Orientation,
+		InPaintParams.UVRange
+	);
+
+	InPaintParams.DrawElements.PushClip(ClippingZone);
+
+	PaintPanel(InPaintParams, Geometry);
+	PaintImage(InPaintParams, Geometry);
+
+	InPaintParams.DrawElements.PopClip();
+}
+
+void FMediaImageViewer::PaintPanel(FMediaImagePaintParams& InPaintParams, const FMediaImagePaintGeometry& InPaintGeometry)
+{
+	const FLinearColor AlphaTint = FLinearColor(1.f, 1.f, 1.f, InPaintParams.ImageOpacity);
+
+	if (GetPanelSettings().BackgroundColor.IsSet())
+	{
+		FSlateDrawElement::MakeBox(
+			InPaintParams.DrawElements,
+			InPaintParams.LayerId,
+			InPaintGeometry.Geometry,
+			&BackgroundColorBrush,
+			ESlateDrawEffect::NoPixelSnapping,
+			GetPanelSettings().BackgroundColor.GetValue() * AlphaTint
+		);
+
+		++InPaintParams.LayerId;
+	}
+
+	if (BackgroundImageBrush.GetResourceObject())
+	{
+		FSlateDrawElement::MakeBox(
+			InPaintParams.DrawElements,
+			InPaintParams.LayerId,
+			InPaintGeometry.Geometry,
+			&BackgroundImageBrush,
+			ESlateDrawEffect::NoPixelSnapping,
+			AlphaTint
+		);
+
+		++InPaintParams.LayerId;
+	}
+}
+
+void FMediaImageViewer::PaintImage(FMediaImagePaintParams& InPaintParams, const FMediaImagePaintGeometry& InPaintGeometry)
+{
+	FSlateDrawElement::MakeBox(
+		InPaintParams.DrawElements,
+		InPaintParams.LayerId,
+		InPaintGeometry.Geometry,
+		Brush.Get(),
+		ESlateDrawEffect::NoPixelSnapping,
+		GetPaintSettings().Tint * FLinearColor(1.f, 1.f, 1.f, InPaintParams.ImageOpacity)
+	);
+
+	++InPaintParams.LayerId;
+}
+
+} // UE::MediaViewer
diff --git a/Engine/Plugins/Experimental/MediaViewer/Source/MediaViewer/Private/Library/MediaViewerLibraryIni.cpp b/Engine/Plugins/Experimental/MediaViewer/Source/MediaViewer/Private/Library/MediaViewerLibraryIni.cpp
new file mode 100644
--- /dev/null
+++ b/Engine/Plugins/Experimental/MediaViewer/Source/MediaViewer/Private/Library/MediaViewerLibraryIni.cpp
@@ -0,0 +1,114 @@
+// Copyright Epic Games, Inc. All Rights Reserved.
+
+#include "Library/MediaViewerLibraryIni.h"
+
+#include "IMediaViewerModule.h"
+#include "Library/MediaViewerLibrary.h"
+#include "Library/MediaViewerLibraryGroup.h"
+#include "Library/MediaViewerLibraryItem.h"
+#include "Templates/SharedPointer.h"
+
+UMediaViewerLibraryIni& UMediaViewerLibraryIni::Get()
+{
+	return *GetMutableDefault<UMediaViewerLibraryIni>();
+}
+
+void UMediaViewerLibraryIni::SaveLibrary(const TSharedRef<UE::MediaViewer::Private::FMediaViewerLibrary>& InLibrary)
+{
+	Groups.Empty();
+	Groups.Reserve(InLibrary->Groups.Num());
+
+	for (const TSharedRef<const FMediaViewerLibraryGroup>& Group : InLibrary->GetGroups())
+	{
+		if (Group->IsDynamic())
+		{
+			continue;
+		}
+
+		Groups.Add(*Group);
+	}
+
+	Items.Empty();
+	Items.Reserve(InLibrary->Items.Num());
+
+	for (const TPair<FGuid, TSharedRef<FMediaViewerLibraryItem>>& ItemPair : InLibrary->Items)
+	{
+		const TSharedRef<FMediaViewerLibraryItem>& Item = ItemPair.Value;
+
+		if (Item->IsTransient())
+		{
+			continue;
+		}
+
+		Items.Emplace(Item->GetItemType(), *Item);
+	}
+
+	SaveConfig();
+}
+
+void UMediaViewerLibraryIni::LoadLibrary(const TSharedRef<UE::MediaViewer::Private::FMediaViewerLibrary>& InLibrary) const
+{
+	UE::MediaViewer::IMediaViewerModule& Module = UE::MediaViewer::IMediaViewerModule::Get();
+
+	for (const FMediaViewerLibraryItemData& ItemData : Items)
+	{
+		const FGuid ItemId = ItemData.Item.GetId();
+
+		if (!InLibrary->Items.Contains(ItemId))
+		{
+			if (TSharedPtr<FMediaViewerLibraryItem> Item = Module.CreateLibraryItem(ItemData.ItemType, ItemData.Item))
+			{
+				InLibrary->Items.Add(ItemId, Item.ToSharedRef());
+			}
+		}
+	}
+
+	for (const FMediaViewerLibraryGroup& SavedGroup : Groups)
+	{
+		TSharedPtr<FMediaViewerLibraryGroup> Group = InLibrary->GetGroup(SavedGroup.GetId());
+
+		if (!Group.IsValid())
+		{
+			Group = MakeShared<FMediaViewerLibraryGroup>(FMediaViewerLibraryGroup::FPrivateToken(), SavedGroup);
+			InLibrary->AddGroup(Group.ToSharedRef());
+		}
+
+		TSet<FGuid> ExistingItems;
+		ExistingItems.Append(Group->Items);
+
+		for (const FGuid& ItemId : SavedGroup.Items)
+		{
+			if (!ExistingItems.Contains(ItemId))
+			{
+				Group->Items.Add(ItemId);
+				ExistingItems.Add(ItemId);
+			}
+		}
+	}
+}
+
+bool UMediaViewerLibraryIni::HasGroup(const FGuid& InGroupId) const
+{
+	for (const FMediaViewerLibraryGroup& Group : Groups)
+	{
+		if (Group.GetId() == InGroupId)
+		{
+			return true;
+		}
+	}
+
+	return false;
+}
+
+bool UMediaViewerLibraryIni::HasItem(const FGuid& InItemId) const
+{
+	for (const FMediaViewerLibraryItemData& ItemData : Items)
+	{
+		if (ItemData.Item.GetId() == InItemId)
+		{
+			return true;
+		}
+	}
+
+	return false;
+}
diff --git a/Engine/Plugins/Experimental/MediaViewer/Source/MediaViewer/Private/Library/MediaViewerLibraryMediaTextureGroup.cpp b/Engine/Plugins/Experimental/MediaViewer/Source/MediaViewer/Private/Library/MediaViewerLibraryMediaTextureGroup.cpp
new file mode 100644
--- /dev/null
+++ b/Engine/Plugins/Experimental/MediaViewer/Source/MediaViewer/Private/Library/MediaViewerLibraryMediaTextureGroup.cpp
@@ -0,0 +1,56 @@
+// Copyright Epic Games, Inc. All Rights Reserved.
+
+#include "Library/MediaViewerLibraryMediaTextureGroup.h"
+
+#include "ImageViewers/MediaTextureImageViewer.h"
+#include "Library/MediaViewerLibraryItem.h"
+#include "MediaTexture.h"
+#include "UObject/UObjectIterator.h"
+
+#define LOCTEXT_NAMESPACE "MediaViewerLibraryMediaTextureGroup"
+
+namespace UE::MediaViewer::Private
+{
+
+FMediaViewerLibraryMediaTextureGroup::FMediaViewerLibraryMediaTextureGroup(const TSharedRef<IMediaViewerLibrary>& InLibrary)
+	: FMediaViewerLibraryMediaTextureGroup(InLibrary, FGuid::NewGuid())
+{
+}
+
+FMediaViewerLibraryMediaTextureGroup::FMediaViewerLibraryMediaTextureGroup(const TSharedRef<IMediaViewerLibrary>& InLibrary,
+	const FGuid& InGuid)
+	: FMediaViewerLibraryDynamicGroup(
+		InLibrary,
+		InGuid,
+		LOCTEXT("MediaTextures", "Media Textures"),
+		LOCTEXT("MediaTexturesTooltip", "The available media textures."),
+		FGenerateItems::CreateStatic(&FMediaViewerLibraryMediaTextureGroup::GetMediaTextureItems)
+	)
+{
+}
+
+TArray<TSharedRef<FMediaViewerLibraryItem>> FMediaViewerLibraryMediaTextureGroup::GetMediaTextureItems()
+{
+	TArray<TSharedRef<FMediaViewerLibraryItem>> MediaTextureItems;
+
+	TSharedRef<FMediaTextureImageViewer::FFactory> Factory = MakeShared<FMediaTextureImageViewer::FFactory>();
+
+	for (UMediaTexture* MediaTexture : TObjectRange<UMediaTexture>())
+	{
+		if (MediaTexture->IsTemplate())
+		{
+			continue;
+		}
+
+		if (TSharedPtr<FMediaViewerLibraryItem> NewItem = Factory->CreateLibraryItem(TNotNull<UMediaTexture*>(MediaTexture)))
+		{
+			MediaTextureItems.Add(NewItem.ToSharedRef());
+		}
+	}
+
+	return MediaTextureItems;
+}
+
+} // UE::MediaViewer::Private
+
+#undef LOCTEXT_NAMESPACE
diff --git a/Engine/Plugins/Experimental/MediaViewer/Source/MediaViewer/Private/Library/MediaViewerLibraryMediaTextureGroup.h b/Engine/Plugins/Experimental/MediaViewer/Source/MediaViewer/Private/Library/MediaViewerLibraryMediaTextureGroup.h
new file mode 100644
--- /dev/null
+++ b/Engine/Plugins/Experimental/MediaViewer/Source/MediaViewer/Private/Library/MediaViewerLibraryMediaTextureGroup.h
@@ -0,0 +1,24 @@
+// Copyright Epic Games, Inc. All Rights Reserved.
+
+#pragma once
+
+#include "Library/MediaViewerLibraryDynamicGroup.h"
+
+namespace UE::MediaViewer::Private
+{
+
+/**
+ * A group that generates entries based on available UMediaTextures.
+*/
+struct FMediaViewerLibraryMediaTextureGroup : public FMediaViewerLibraryDynamicGroup
+{
+	FMediaViewerLibraryMediaTextureGroup(const TSharedRef<IMediaViewerLibrary>& InLibrary);
+	FMediaViewerLibraryMediaTextureGroup(const TSharedRef<IMediaViewerLibrary>& InLibrary, const FGuid& InId);
+
+	virtual ~FMediaViewerLibraryMediaTextureGroup() override = default;
+
+protected:
+	static TArray<TSharedRef<FMediaViewerLibraryItem>> GetMediaTextureItems();
+};
+
+} // UE::MediaViewer::Private
diff --git a/Engine/Plugins/Experimental/MediaViewer/Source/MediaViewer/Private/Library/MediaViewerLibraryTextureRenderTarget2DGroup.cpp b/Engine/Plugins/Experimental/MediaViewer/Source/MediaViewer/Private/Library/MediaViewerLibraryTextureRenderTarget2DGroup.cpp
new file mode 100644
--- /dev/null
+++ b/Engine/Plugins/Experimental/MediaViewer/Source/MediaViewer/Private/Library/MediaViewerLibraryTextureRenderTarget2DGroup.cpp
@@ -0,0 +1,57 @@
+// Copyright Epic Games, Inc. All Rights Reserved.
+
+#include "Library/MediaViewerLibraryTextureRenderTarget2DGroup.h"
+
+#include "Engine/TextureRenderTarget2D.h"
+#include "ImageViewers/TextureRenderTarget2DImageViewer.h"
+#include "Library/MediaViewerLibraryItem.h"
+#include "MediaViewerUtils.h"
+#include "UObject/UObjectIterator.h"
+
+#define LOCTEXT_NAMESPACE "MediaViewerLibraryTextureRenderTarget2DGroup"
+
+namespace UE::MediaViewer::Private
+{
+
+FMediaViewerLibraryTextureRenderTarget2DGroup::FMediaViewerLibraryTextureRenderTarget2DGroup(const TSharedRef<IMediaViewerLibrary>& InLibrary)
+	: FMediaViewerLibraryTextureRenderTarget2DGroup(InLibrary, FGuid::NewGuid())
+{
+}
+
+FMediaViewerLibraryTextureRenderTarget2DGroup::FMediaViewerLibraryTextureRenderTarget2DGroup(const TSharedRef<IMediaViewerLibrary>& InLibrary, 
+	const FGuid& InGuid)
+	: FMediaViewerLibraryDynamicGroup(
+		InLibrary,
+		InGuid,
+		LOCTEXT("RenderTargets", "Render Targets"),
+		LOCTEXT("RenderTargetsTooltip", "The available render targets."),
+		FGenerateItems::CreateStatic(&FMediaViewerLibraryTextureRenderTarget2DGroup::GetTextureRenderTarget2DItems)
+	)
+{
+}
+
+TArray<TSharedRef<FMediaViewerLibraryItem>> FMediaViewerLibraryTextureRenderTarget2DGroup::GetTextureRenderTarget2DItems()
+{
+	TArray<TSharedRef<FMediaViewerLibraryItem>> TextureRenderTarget2DItems;
+
+	TSharedRef<FTextureRenderTarget2DImageViewer::FFactory> Factory = MakeShared<FTextureRenderTarget2DImageViewer::FFactory>();
+
+	for (UTextureRenderTarget2D* TextureRenderTarget2D : TObjectRange<UTextureRenderTarget2D>())
+	{
+		if (TextureRenderTarget2D->IsTemplate())
+		{
+			continue;
+		}
+
+		if (TSharedPtr<FMediaViewerLibraryItem> NewItem = Factory->CreateLibraryItem(TNotNull<UTextureRenderTarget2D*>(TextureRenderTarget2D)))
+		{
+			TextureRenderTarget2DItems.Add(NewItem.ToSharedRef());
+		}
+	}
+
+	return TextureRenderTarget2DItems;
+}
+
+} // UE::MediaViewer::Private
+
+#undef LOCTEXT_NAMESPACE
diff --git a/Engine/Plugins/Experimental/MediaViewer/Source/MediaViewer/Private/Library/MediaViewerLibraryTextureRenderTarget2DGroup.h b/Engine/Plugins/Experimental/MediaViewer/Source/MediaViewer/Private/Library/MediaViewerLibraryTextureRenderTarget2DGroup.h
new file mode 100644
--- /dev/null
+++ b/Engine/Plugins/Experimental/MediaViewer/Source/MediaViewer/Private/Library/MediaViewerLibraryTextureRenderTarget2DGroup.h
@@ -0,0 +1,24 @@
+// Copyright Epic Games, Inc. All Rights Reserved.
+
+#pragma once
+
+#include "Library/MediaViewerLibraryDynamicGroup.h"
+
+namespace UE::MediaViewer::Private
+{
+
+/**
+ * A group that generates entries based on available UTextureRenderTarget2Ds.
+*/
+struct FMediaViewerLibraryTextureRenderTarget2DGroup : public FMediaViewerLibraryDynamicGroup
+{
+	FMediaViewerLibraryTextureRenderTarget2DGroup(const TSharedRef<IMediaViewerLibrary>& InLibrary);
+	FMediaViewerLibraryTextureRenderTarget2DGroup(const TSharedRef<IMediaViewerLibrary>& InLibrary, const FGuid& InId);
+
+	virtual ~FMediaViewerLibraryTextureRenderTarget2DGroup() override = default;
+
+protected:
+	static TArray<TSharedRef<FMediaViewerLibraryItem>> GetTextureRenderTarget2DItems();
+};
+
+} // UE::MediaViewer::Private
diff --git a/Engine/Plugins/Experimental/MediaViewer/Source/MediaViewer/Private/Widgets/MediaViewerSettings.h b/Engine/Plugins/Experimental/MediaViewer/Source/MediaViewer/Private/Widgets/MediaViewerSettings.h
new file mode 100644
--- /dev/null
+++ b/Engine/Plugins/Experimental/MediaViewer/Source/MediaViewer/Private/Widgets/MediaViewerSettings.h
@@ -0,0 +1,58 @@
+// Copyright Epic Games, Inc. All Rights Reserved.
+
+#pragma once
+
+#include "Containers/EnumAsByte.h"
+#include "Math/Color.h"
+#include "Math/Vector.h"
+#include "Styling/StyleColors.h"
+#include "Types/SlateEnums.h"
+#include "UObject/ObjectPtr.h"
+
+#include "MediaViewerSettings.generated.h"
+
+class UTexture2D;
+
+/** Settings for the viewer. */
+USTRUCT()
+struct FMediaViewerSettings
+{
+	GENERATED_BODY()
+
+	/** Color painted to the viewer behind the images. */
+	UPROPERTY(EditAnywhere, Category = "Background")
+	FLinearColor ClearColor = FStyleColors::Panel.GetSpecifiedColor();
+
+	/** Texture tiled behind the images, on top of the clear color. */
+	UPROPERTY(EditAnywhere, Category = "Background")
+	TObjectPtr<UTexture2D> Texture = nullptr;
+
+	/** Offset of the tiled background texture. */
+	UPROPERTY(EditAnywhere, Category = "Background", meta = (EditCondition = "Texture != nullptr"))
+	FVector2D Offset = FVector2D::ZeroVector;
+
+	/** Scale of the tiled background texture. */
+	UPROPERTY(EditAnywhere, Category = "Background", meta = (UIMin = 0.1, ClampMin = 0.1, UIMax = 10, ClampMax = 10,
+		EditCondition = "Texture != nullptr"))
+	float Scale = 1.f;
+
+	/** If true, when one image transform changes, the other image transforms will change an equal amount. */
+	UPROPERTY(EditAnywhere, Category = "AB View", DisplayName = "Lock Transforms")
+	bool bAreTransformsLocked = false;
+
+	/** When in AB view, whether the viewers are split horizontally or vertically. */
+	TEnumAsByte<EOrientation> ABOrientation = Orient_Horizontal;
+
+	/** Splitter location (0-100) */
+	UPROPERTY(EditAnywhere, Category = "AB View", DisplayName = "AB Splitter Location",
+		meta = (UIMin = 0, ClampMin = 0, UIMax = 100, ClampMax = 100, Units = "Percent"))
+	float ABSplitterLocation = 50.f;
+
+	/** Opacity of the second/B image (0-100).
+	 * At 100%, the first/A image is rendered only on the left/top side.
+	 * Below 100%, the first/A image is rendered on the entire panel with B overlaid on top of it, translucently.
+	 */
+	UPROPERTY(EditAnywhere, Category = "AB View", DisplayName = "B Image Opacity",
+		meta = (UIMin = 0, ClampMin = 0, UIMax = 100, ClampMax = 100, Units = "Percent"))
+	float SecondImageOpacity = 100.f;
+};
diff --git a/Engine/Plugins/Experimental/MediaViewer/Source/MediaViewer/Private/Widgets/SMediaImageViewerOverlay.cpp b/Engine/Plugins/Experimental/MediaViewer/Source/MediaViewer/Private/Widgets/SMediaImageViewerOverlay.cpp
new file mode 100644
--- /dev/null
+++ b/Engine/Plugins/Experimental/MediaViewer/Source/MediaViewer/Private/Widgets/SMediaImageViewerOverlay.cpp
@@ -0,0 +1,919 @@
+// Copyright Epic Games, Inc. All Rights Reserved.
+
+#include "Widgets/SMediaImageViewerOverlay.h"
+
+#include "EditorViewportClient.h"
+#include "Framework/Application/SlateApplication.h"
+#include "Framework/Commands/UICommandList.h"
+#include "Framework/MultiBox/MultiBoxBuilder.h"
+#include "HAL/PlatformApplicationMisc.h"
+#include "ImageViewer/MediaImageViewer.h"
+#include "IMediaViewerModule.h"
+#include "Library/MediaViewerLibrary.h"
+#include "Library/MediaViewerLibraryItem.h"
+#include "MediaViewer.h"
+#include "MediaViewerCommands.h"
+#include "SAssetDropTarget.h"
+#include "Widgets/Layout/SScissorRectBox.h"
+#include "Widgets/SBoxPanel.h"
+#include "Widgets/SMediaImageViewerStatusBar.h"
+#include "Widgets/SMediaViewer.h"
+
+#define LOCTEXT_NAMESPACE "SMediaImageViewerOverlay"
+
+namespace UE::MediaViewer::Private
+{
+
+SMediaImageViewerOverlay::SMediaImageViewerOverlay()
+	: CommandList(MakeShared<FUICommandList>())
+	, DragButtonName(NAME_None)
+	, bDragging(false)
+	, bExternalDragging(false)
+	, bOverlayEnabled(true)
+{	
+}
+
+void SMediaImageViewerOverlay::PrivateRegisterAttributes(FSlateAttributeDescriptor::FInitializer&)
+{
+}
+
+void SMediaImageViewerOverlay::Construct(const FArguments& InArgs, EMediaImageViewerPosition InPosition, 
+	const TSharedPtr<FMediaViewerDelegates>& InDelegates)
+{
+	Position = InPosition;
+	GetImageViewerDelegate = InDelegates->GetImageViewer;
+
+	TSharedPtr<FMediaImageViewer> ImageViewer = GetImageViewerDelegate.Execute(Position);
+
+	if (!ImageViewer.IsValid())
+	{
+		return;
+	}
+
+	CachedItem = ImageViewer->CreateLibraryItem();
+
+	BindCommands();
+	InDelegates->GetCommandList.Execute()->Append(CommandList);
+
+	Delegates = InDelegates;
+
+	const bool bDisplayLeft = InPosition == EMediaImageViewerPosition::First
+		&& InDelegates->GetABOrientation.Execute() == Orient_Horizontal;
+
+	const float PaddingLeft = bDisplayLeft ? 30.f : 0.f;
+	const float PaddingRight = !bDisplayLeft ? 30.f : 0.f;
+
+	ChildSlot
+	[
+		SNew(SScissorRectBox)
+		[
+			SAssignNew(Container, SVerticalBox)
+
+			+ SVerticalBox::Slot()
+			.FillHeight(1.f)
+			.HAlign(HAlign_Fill)
+			.VAlign(VAlign_Fill)
+			[
+				CreateOverlay(InArgs._bCreateSplitDropTarget)
+			]
+
+			+ SVerticalBox::Slot()
+			.AutoHeight()
+			.HAlign(HAlign_Fill)
+			.VAlign(VAlign_Center)
+			.Padding(PaddingLeft, 0.f, PaddingRight, 0.f)
+			[
+				CreateStatusBar(InDelegates)
+			]
+		]
+	];
+}
+
+TSharedPtr<FMediaImageViewer> SMediaImageViewerOverlay::GetImageViewer() const
+{
+	return GetImageViewerDelegate.Execute(Position);
+}
+
+const TSharedRef<FUICommandList>& SMediaImageViewerOverlay::GetCommandList() const
+{
+	return CommandList;
+}
+
+FIntPoint SMediaImageViewerOverlay::GetImageViewerPixelCoordinates() const
+{
+	TSharedPtr<FMediaImageViewer> ImageViewer = GetImageViewerDelegate.Execute(Position);
+
+	if (!ImageViewer.IsValid())
+	{
+		return FIntPoint(-1, -1);
+	}
+
+	const FIntPoint ViewerSize = Delegates->GetViewerSize.Execute();
+	const FIntPoint CursorLocation = Delegates->GetCursorLocation.Execute() - Delegates->GetViewerPosition.Execute();
+	const FVector2D TopLeft = ImageViewer->GetPaintOffset(ViewerSize, Delegates->GetViewerPosition.Execute());
+
+	FVector2D CursorLocationD = FVector2D(CursorLocation.X, CursorLocation.Y) - TopLeft + Delegates->GetViewerPosition.Execute();
+	CursorLocationD /= ImageViewer->GetPaintSettings().Scale;
+
+	return FIntPoint(
+		FMath::RoundToInt(CursorLocationD.X),
+		FMath::RoundToInt(CursorLocationD.Y)
+	);
+}
+
+bool SMediaImageViewerOverlay::IsCursorOverImageViewer() const
+{
+	TSharedPtr<FMediaImageViewer> ImageViewer = GetImageViewerDelegate.Execute(Position);
+
+	if (!ImageViewer.IsValid())
+	{
+		return false;
+	}
+
+	const FIntPoint PixelCoordinates = GetImageViewerPixelCoordinates();
+
+	if (PixelCoordinates.X < 0 || PixelCoordinates.Y < 0)
+	{
+		return false;
+	}
+
+	const FIntPoint& ImageSize = ImageViewer->GetInfo().Size;
+
+	if (PixelCoordinates.X >= ImageSize.X || PixelCoordinates.Y >= ImageSize.Y)
+	{
+		return false;
+	}
+
+	return true;
+}
+
+FCursorReply SMediaImageViewerOverlay::OnCursorQuery(const FGeometry& InMyGeometry, const FPointerEvent& InCursorEvent) const
+{
+	if (bDragging || bExternalDragging)
+	{
+		return FCursorReply::Cursor(EMouseCursor::GrabHandClosed);
+	}
+
+	if (IsCursorOverImageViewer())
+	{
+		return FCursorReply::Cursor(EMouseCursor::Crosshairs);
+	}
+
+	return SCompoundWidget::OnCursorQuery(InMyGeometry, InCursorEvent);
+}
+
+FReply SMediaImageViewerOverlay::OnKeyDown(const FGeometry& InMyGeometry, const FKeyEvent& InKeyEvent)
+{
+	if (CommandList->ProcessCommandBindings(InKeyEvent))
+	{
+		return FReply::Handled();
+	}
+
+	return SCompoundWidget::OnKeyDown(InMyGeometry, InKeyEvent);
+}
+
+FReply SMediaImageViewerOverlay::OnMouseButtonDown(const FGeometry& InMyGeometry, const FPointerEvent& InMouseEvent)
+{
+	if (!bDragging && !bExternalDragging)
+	{
+		const FIntPoint CursorLocation = Delegates->GetCursorLocation.Execute();
+
+		if (InMouseEvent.GetEffectingButton() == EKeys::LeftMouseButton)
+		{
+			DragStartCursor = CursorLocation;
+			DragStartOffset = GetOffset();
+			DragButtonName = InMouseEvent.GetEffectingButton().GetFName();
+			bDragging = true;
+			return FReply::Handled();
+		}
+
+		if (TSharedPtr<FMediaImageViewer> ImageViewer = GetImageViewerDelegate.Execute(Position))
+		{
+			const FInputEventState EventState(nullptr, InMouseEvent.GetEffectingButton(), IE_Pressed);
+
+			if (ImageViewer->OnTrackingStarted(EventState, CursorLocation))
+			{
+				DragButtonName = InMouseEvent.GetEffectingButton().GetFName();
+				bExternalDragging = true;
+				return FReply::Handled();
+			}
+		}
+	}
+
+	return SCompoundWidget::OnMouseButtonDown(InMyGeometry, InMouseEvent);
+}
+
+FReply SMediaImageViewerOverlay::OnMouseButtonUp(const FGeometry& InMyGeometry, const FPointerEvent& InMouseEvent)
+{
+	if (!bDragging && !bExternalDragging)
+	{
+		if (InMouseEvent.GetEffectingButton() == EKeys::RightMouseButton)
+		{
+			FSlateApplication::Get().PushMenu(
+				SharedThis(this),
+				FWidgetPath(),
+				CreateMenu(),
+				FSlateApplication::Get().GetCursorPos(),
+				FPopupTransitionEffect::ContextMenu
+			);
+
+			return FReply::Handled();
+		}
+	}
+	else
+	{
+		OnDragButtonUp(InMouseEvent.GetEffectingButton().GetFName());
+	}
+
+	return SCompoundWidget::OnMouseButtonUp(InMyGeometry, InMouseEvent);
+}
+
+FReply SMediaImageViewerOverlay::OnMouseWheel(const FGeometry& InMyGeometry, const FPointerEvent& InMouseEvent)
+{
+	static const float NormalMultiplier = FMath::Pow(2.0, 0.125);
+	static const float FastMultiplier = FMath::Pow(2.0, 0.5);
+	static const float SlowMultiplier = FMath::Pow(2.0, 0.03125);
+
+	const FModifierKeysState ModifierKeys = FSlateApplication::Get().GetModifierKeys();
+	
+	const float ModifierMultiplier = 
+		ModifierKeys.IsLeftShiftDown() 
+		? FastMultiplier
+		: ((ModifierKeys.IsLeftControlDown() || ModifierKeys.IsLeftCommandDown())
+			? SlowMultiplier
+			: NormalMultiplier);
+	
+	const float WheelDelta = InMouseEvent.GetWheelDelta();
+
+	if (WheelDelta > 0)
+	{
+		Try_SetScale(GetScale() * ModifierMultiplier);
+	}
+	else if (WheelDelta < 0)
+	{
+		Try_SetScale(GetScale() / ModifierMultiplier);
+	}
+
+	return SCompoundWidget::OnMouseWheel(InMyGeometry, InMouseEvent);
+}
+
+bool SMediaImageViewerOverlay::SupportsKeyboardFocus() const
+{
+	return true;
+}
+
+void SMediaImageViewerOverlay::UpdateMouse(const TOptional<FIntPoint>& InMousePosition)
+{
+	if (!DragButtonName.IsNone())
+	{
+		if (!FSlateApplication::Get().GetPressedMouseButtons().Contains(DragButtonName))
+		{
+			OnDragButtonUp(DragButtonName);
+		}
+	}
+
+	if (bDragging)
+	{
+		UpdateDragPosition();
+	}
+
+	if (InMousePosition.IsSet())
+	{
+		if (TSharedPtr<FMediaImageViewer> ImageViewer = GetImageViewerDelegate.Execute(Position))
+		{
+			ImageViewer->OnMouseMove(InMousePosition.GetValue());
+		}
+	}
+}
+
+void SMediaImageViewerOverlay::OnDragButtonUp(FName InKeyName)
+{
+	if (InKeyName != DragButtonName)
+	{
+		return;
+	}
+
+	if (bDragging)
+	{
+		UpdateDragPosition();
+		DragButtonName = NAME_None;
+		bDragging = false;
+	}
+	else if (bExternalDragging)
+	{
+		if (TSharedPtr<FMediaImageViewer> ImageViewer = GetImageViewerDelegate.Execute(Position))
+		{
+			ImageViewer->OnTrackingStopped(Delegates->GetCursorLocation.Execute());
+		}
+
+		DragButtonName = NAME_None;
+		bExternalDragging = false;
+	}
+}
+
+void SMediaImageViewerOverlay::UpdateDragPosition()
+{
+	const FIntPoint CursorLocation = Delegates->GetCursorLocation.Execute();
+
+	FVector NewOffset = DragStartOffset;
+	NewOffset.X += (CursorLocation.X - DragStartCursor.X) /*/ ImageViewer->GetPaintSettings().Scale*/;
+	NewOffset.Y += (CursorLocation.Y - DragStartCursor.Y) /*/ ImageViewer->GetPaintSettings().Scale*/;
+
+	Try_AddOffset(NewOffset - GetOffset());
+}
+
+bool SMediaImageViewerOverlay::OnAssetDraggedOver(TArrayView<FAssetData> InAssetData, FText& OutError) const
+{
+	IMediaViewerModule& Module = IMediaViewerModule::Get();
+
+	for (const FAssetData& AssetData : InAssetData)
+	{
+		if (Module.HasFactoryFor(AssetData))
+		{
+			return true;
+		}
+	}
+
+	OutError = LOCTEXT("NotSupported", "Asset is not supported.");
+
+	return false;
+}
+
+void SMediaImageViewerOverlay::OnAssetsDropped(const FDragDropEvent& InDragDropEvent, TArrayView<FAssetData> InAssetData, 
+	EMediaImageViewerPosition InPosition)
+{
+	IMediaViewerModule& Module = IMediaViewerModule::Get();
+
+	TArray<int32> Indices;
+	Indices.Reserve(static_cast<int32>(EMediaImageViewerPosition::COUNT));
+
+	for (int32 Index = 0; Index < InAssetData.Num(); ++Index)
+	{
+		if (Module.HasFactoryFor(InAssetData[Index]))
+		{
+			Indices.Add(Index);
+
+			if (Indices.Num() == static_cast<int32>(EMediaImageViewerPosition::COUNT))
+			{
+				break;
+			}
+		}
+	}
+
+	if (Indices.IsEmpty())
+	{
+		return;
+	}
+
+	if (Indices.Num() == 1)
+	{
+		Module.SetImage(InPosition, InAssetData[Indices[0]]);
+	}
+	else
+	{
+		Module.SetImage(EMediaImageViewerPosition::First, InAssetData[Indices[0]]);
+		Module.SetImage(EMediaImageViewerPosition::Second, InAssetData[Indices[1]]);
+	}
+}
+
+bool SMediaImageViewerOverlay::OnVerifyDrag(TSharedPtr<FDragDropOperation> InDragDropOp)
+{
+	return true;
+}
+
+FReply SMediaImageViewerOverlay::OnDrop(const FGeometry& InMyGeometry, const FDragDropEvent& InDragDropEvent)
+{
+	return SCompoundWidget::OnDrop(InMyGeometry, InDragDropEvent);
+}
+
+EVisibility SMediaImageViewerOverlay::GetDragDescriptionVisibility() const
+{
+	return FSlateApplication::Get().IsDragDropping()
+		? EVisibility::HitTestInvisible
+		: EVisibility::Collapsed;
+}
+
+void SMediaImageViewerOverlay::BindCommands()
+{
+	const FMediaViewerCommands& Commands = FMediaViewerCommands::Get();
+
+	CommandList->MapAction(Commands.MoveLeft,     FExecuteAction::CreateSP(this, &SMediaImageViewerOverlay::Try_AddOffset, FVector(-10,   0,   0)));
+	CommandList->MapAction(Commands.MoveRight,    FExecuteAction::CreateSP(this, &SMediaImageViewerOverlay::Try_AddOffset, FVector( 10,   0,   0)));
+	CommandList->MapAction(Commands.MoveUp,       FExecuteAction::CreateSP(this, &SMediaImageViewerOverlay::Try_AddOffset, FVector(  0, -10,   0)));
+	CommandList->MapAction(Commands.MoveDown,     FExecuteAction::CreateSP(this, &SMediaImageViewerOverlay::Try_AddOffset, FVector(  0,  10,   0)));
+	CommandList->MapAction(Commands.MoveForward,  FExecuteAction::CreateSP(this, &SMediaImageViewerOverlay::Try_AddOffset, FVector(  0,   0, -10)));
+	CommandList->MapAction(Commands.MoveBackward, FExecuteAction::CreateSP(this, &SMediaImageViewerOverlay::Try_AddOffset, FVector(  0,   0,  10)));
+
+	CommandList->MapAction(Commands.RotateMinusPitch, FExecuteAction::CreateSP(this, &SMediaImageViewerOverlay::Try_AddRotation, FRotator(-10,   0,   0)));
+	CommandList->MapAction(Commands.RotatePlusPitch,  FExecuteAction::CreateSP(this, &SMediaImageViewerOverlay::Try_AddRotation, FRotator( 10,   0,   0)));
+	CommandList->MapAction(Commands.RotateMinusYaw,   FExecuteAction::CreateSP(this, &SMediaImageViewerOverlay::Try_AddRotation, FRotator(  0, -10,   0)));
+	CommandList->MapAction(Commands.RotatePlusYaw,    FExecuteAction::CreateSP(this, &SMediaImageViewerOverlay::Try_AddRotation, FRotator(  0,  10,   0)));
+	CommandList->MapAction(Commands.RotateMinusRoll,  FExecuteAction::CreateSP(this, &SMediaImageViewerOverlay::Try_AddRotation, FRotator(  0,   0, -10)));
+	CommandList->MapAction(Commands.RotatePlusRoll,   FExecuteAction::CreateSP(this, &SMediaImageViewerOverlay::Try_AddRotation, FRotator(  0,   0,  10)));
+
+	CommandList->MapAction(Commands.Scale12,  FExecuteAction::CreateSP(this, &SMediaImageViewerOverlay::Try_SetScale, 0.125f));
+	CommandList->MapAction(Commands.Scale25,  FExecuteAction::CreateSP(this, &SMediaImageViewerOverlay::Try_SetScale, 0.25f));
+	CommandList->MapAction(Commands.Scale50,  FExecuteAction::CreateSP(this, &SMediaImageViewerOverlay::Try_SetScale, 0.5f));
+	CommandList->MapAction(Commands.Scale100, FExecuteAction::CreateSP(this, &SMediaImageViewerOverlay::Try_SetScale, 1.0f));
+	CommandList->MapAction(Commands.Scale200, FExecuteAction::CreateSP(this, &SMediaImageViewerOverlay::Try_SetScale, 2.0f));
+	CommandList->MapAction(Commands.Scale400, FExecuteAction::CreateSP(this, &SMediaImageViewerOverlay::Try_SetScale, 4.0f));
+	CommandList->MapAction(Commands.Scale800, FExecuteAction::CreateSP(this, &SMediaImageViewerOverlay::Try_SetScale, 8.0f));
+
+	CommandList->MapAction(Commands.ScaleToFit, FExecuteAction::CreateSP(this, &SMediaImageViewerOverlay::ScaleToFit));
+
+	CommandList->MapAction(
+		Commands.ResetTransform, 
+		FExecuteAction::CreateSP(this, &SMediaImageViewerOverlay::Try_ResetTransform),
+		FCanExecuteAction(),
+		FGetActionCheckState(),
+		FIsActionButtonVisible::CreateSP(this, &SMediaImageViewerOverlay::CanResetTransform)
+	);
+
+	CommandList->MapAction(
+		Commands.CopyTransform, 
+		FExecuteAction::CreateSP(this, &SMediaImageViewerOverlay::CopyTransform),
+		FCanExecuteAction(),
+		FGetActionCheckState(),
+		FIsActionButtonVisible::CreateSP(this, &SMediaImageViewerOverlay::CanCopyTransform)
+	);
+
+	CommandList->MapAction(Commands.MipMinus, FExecuteAction::CreateSP(this, &SMediaImageViewerOverlay::AdjustMipLevel, static_cast<int8>(-1)));
+	CommandList->MapAction(Commands.MipPlus,  FExecuteAction::CreateSP(this, &SMediaImageViewerOverlay::AdjustMipLevel, static_cast<int8>(1)));
+
+	CommandList->MapAction(Commands.CopyColor, FExecuteAction::CreateSP(this, &SMediaImageViewerOverlay::CopyColor));
+
+	CommandList->MapAction(
+		Commands.AddToLibrary,
+		FExecuteAction::CreateSP(this, &SMediaImageViewerOverlay::AddToLibrary),
+		FCanExecuteAction::CreateSP(this, &SMediaImageViewerOverlay::CanAddToLibrary)
+	);
+}
+
+TSharedRef<SWidget> SMediaImageViewerOverlay::CreateOverlay(bool bInCreateSplitDropTarget)
+{
+	TSharedPtr<FMediaImageViewer> ImageViewer = GetImageViewerDelegate.Execute(Position);
+
+	if (!ImageViewer.IsValid())
+	{
+		return SNullWidget::NullWidget;
+	}
+
+	TSharedPtr<SWidget> ImageViewerOverlay = ImageViewer.IsValid() ? ImageViewer->GetOverlayWidget(Position) : nullptr;
+
+	if (!bInCreateSplitDropTarget)
+	{
+		return SNew(SAssetDropTarget)
+			.OnAreAssetsAcceptableForDropWithReason(this, &SMediaImageViewerOverlay::OnAssetDraggedOver)
+			.OnAssetsDropped(this, &SMediaImageViewerOverlay::OnAssetsDropped, Position)
+			.bSupportsMultiDrop(true)
+			[
+				SNew(SOverlay)
+				+ SOverlay::Slot()
+				[
+					ImageViewerOverlay.IsValid() ? ImageViewerOverlay.ToSharedRef() : SNullWidget::NullWidget
+				]
+				+ SOverlay::Slot()
+				[
+					SNew(SBox)
+					.HAlign(HAlign_Center)
+					.VAlign(VAlign_Center)
+					[
+						SNew(STextBlock)
+						.Visibility(this, &SMediaImageViewerOverlay::GetDragDescriptionVisibility)
+						.Text(LOCTEXT("DropAsset", "Drop supported asset or library item here."))
+					]
+				]
+			];
+	}
+
+	TSharedRef<SWidget> FirstDropTarget = SNew(SAssetDropTarget)
+		.OnAreAssetsAcceptableForDropWithReason(this, &SMediaImageViewerOverlay::OnAssetDraggedOver)
+		.OnAssetsDropped(this, &SMediaImageViewerOverlay::OnAssetsDropped, EMediaImageViewerPosition::First)
+		.bSupportsMultiDrop(true)
+		[
+			SNew(SVerticalBox)
+			+ SVerticalBox::Slot()
+			.HAlign(HAlign_Center)
+			.VAlign(VAlign_Bottom)
+			.FillHeight(1.f)
+			.Padding(0.f, 0.f, 0.f, 5.f)
+			[
+				SNew(STextBlock)
+				.Visibility(this, &SMediaImageViewerOverlay::GetDragDescriptionVisibility)
+				.Text(LOCTEXT("ImageA", "Replace Image"))
+			]
+			+ SVerticalBox::Slot()
+			.HAlign(HAlign_Center)
+			.VAlign(VAlign_Top)
+			.FillHeight(1.f)
+			.Padding(0.f, 5.f, 0.f, 0.f)
+			[
+				SNew(STextBlock)
+				.Visibility(this, &SMediaImageViewerOverlay::GetDragDescriptionVisibility)
+				.Text(LOCTEXT("DropAsset", "Drop supported asset or library item here."))
+			]
+		];
+
+	TSharedRef<SWidget> SecondDropTarget = SNew(SAssetDropTarget)
+		.OnAreAssetsAcceptableForDropWithReason(this, &SMediaImageViewerOverlay::OnAssetDraggedOver)
+		.OnAssetsDropped(this, &SMediaImageViewerOverlay::OnAssetsDropped, EMediaImageViewerPosition::Second)
+		.bSupportsMultiDrop(true)
+		[
+			SNew(SVerticalBox)
+			+ SVerticalBox::Slot()
+			.HAlign(HAlign_Center)
+			.VAlign(VAlign_Bottom)
+			.FillHeight(1.f)
+			.Padding(0.f, 0.f, 0.f, 5.f)
+			[
+				SNew(STextBlock)
+				.Visibility(this, &SMediaImageViewerOverlay::GetDragDescriptionVisibility)
+				.Text(LOCTEXT("ImageB", "Add Image"))
+			]
+			+ SVerticalBox::Slot()
+			.HAlign(HAlign_Center)
+			.VAlign(VAlign_Top)
+			.FillHeight(1.f)
+			.Padding(0.f, 5.f, 0.f, 0.f)
+			[
+				SNew(STextBlock)
+				.Visibility(this, &SMediaImageViewerOverlay::GetDragDescriptionVisibility)
+				.Text(LOCTEXT("DropAsset", "Drop supported asset or library item here."))
+			]
+		];
+
+	TSharedPtr<SWidget> SplitDropTarget;
+
+	if (Delegates->GetABOrientation.Execute() == Orient_Horizontal)
+	{
+		SplitDropTarget = SNew(SHorizontalBox)
+			+ SHorizontalBox::Slot()
+			.FillWidth(1.f)
+			.Padding(0.f, 0.f, 2.f, 0.f)
+			[
+				FirstDropTarget
+			]
+
+			+ SHorizontalBox::Slot()
+			.FillWidth(1.f)
+			.Padding(2.f, 0.f, 0.f, 0.f)
+			[
+				SecondDropTarget
+			];
+	}
+	else
+	{
+		SplitDropTarget = SNew(SVerticalBox)
+			+ SVerticalBox::Slot()
+			.FillHeight(1.f)
+			.Padding(0.f, 0.f, 0.f, 2.f)
+			[
+				FirstDropTarget
+			]
+
+			+ SVerticalBox::Slot()
+			.FillHeight(1.f)
+			.Padding(0.f, 2.f, 0.f, 0.f)
+			[
+				SecondDropTarget
+			];
+	}
+
+	return SNew(SOverlay)		
+		+ SOverlay::Slot()
+		[
+			SplitDropTarget.ToSharedRef()
+		]
+		+ SOverlay::Slot()
+		[
+			ImageViewerOverlay.IsValid() ? ImageViewerOverlay.ToSharedRef() : SNullWidget::NullWidget
+		];		
+}
+
+TSharedRef<SWidget> SMediaImageViewerOverlay::CreateStatusBar(const TSharedPtr<FMediaViewerDelegates>& InDelegates)
+{
+	return SNew(SMediaImageViewerStatusBar, Position, InDelegates);
+}
+
+bool SMediaImageViewerOverlay::IsOverlayEnabled() const
+{
+	return bOverlayEnabled;
+}
+
+void SMediaImageViewerOverlay::ToggleOverlay()
+{
+	bOverlayEnabled = !bOverlayEnabled;
+
+	FChildren* Children = Container->GetChildren();
+
+	for (int32 Index = 0; Index < Children->Num(); ++Index)
+	{
+		if (bOverlayEnabled)
+		{
+			Children->GetChildAt(Index)->SetVisibility(EVisibility::Visible);
+		}
+		else
+		{
+			Children->GetChildAt(Index)->SetVisibility(EVisibility::Collapsed);
+		}
+	}
+}
+
+void SMediaImageViewerOverlay::ResetTransform()
+{
+	SetOffset(FVector::ZeroVector);
+	SetRotation(FRotator::ZeroRotator);
+	SetScale(1.f);
+}
+
+FVector SMediaImageViewerOverlay::GetOffset() const
+{
+	if (TSharedPtr<FMediaImageViewer> ImageViewer = GetImageViewerDelegate.Execute(Position))
+	{
+		return ImageViewer->GetPaintSettings().Offset;
+	}
+
+	return FVector::ZeroVector;
+}
+
+void SMediaImageViewerOverlay::Try_AddOffset(FVector InOffset)
+{
+	if (Delegates->AreTransformsLocked.Execute())
+	{
+		Delegates->AddOffsetToAll.Execute(InOffset);
+	}
+	else
+	{
+		AddOffset(InOffset);
+	}
+}
+
+void SMediaImageViewerOverlay::AddOffset(const FVector& InOffset)
+{
+	SetOffset(GetOffset() + InOffset);
+}
+
+void SMediaImageViewerOverlay::SetOffset(const FVector& InOffset)
+{
+	if (TSharedPtr<FMediaImageViewer> ImageViewer = GetImageViewerDelegate.Execute(Position))
+	{
+		ImageViewer->GetPaintSettings().Offset = InOffset;
+	}
+}
+
+FRotator SMediaImageViewerOverlay::GetRotation() const
+{
+	if (TSharedPtr<FMediaImageViewer> ImageViewer = GetImageViewerDelegate.Execute(Position))
+	{
+		return ImageViewer->GetPaintSettings().Rotation;
+	}
+
+	return FRotator::ZeroRotator;
+}
+
+void SMediaImageViewerOverlay::Try_AddRotation(FRotator InRotation)
+{
+	if (Delegates->AreTransformsLocked.Execute())
+	{
+		Delegates->AddRotationToAll.Execute(InRotation);
+	}
+	else
+	{
+		AddRotation(InRotation);
+	}
+}
+
+void SMediaImageViewerOverlay::AddRotation(const FRotator& InRotation)
+{
+	SetRotation(GetRotation() + InRotation);
+}
+
+void SMediaImageViewerOverlay::SetRotation(const FRotator& InRotation)
+{
+	if (TSharedPtr<FMediaImageViewer> ImageViewer = GetImageViewerDelegate.Execute(Position))
+	{
+		ImageViewer->GetPaintSettings().Rotation = InRotation;
+	}
+}
+
+float SMediaImageViewerOverlay::GetScale() const
+{
+	if (TSharedPtr<FMediaImageViewer> ImageViewer = GetImageViewerDelegate.Execute(Position))
+	{
+		return ImageViewer->GetPaintSettings().Scale;
+	}
+	
+	return 1.f;
+}
+
+void SMediaImageViewerOverlay::Try_SetScale(float InScale)
+{
+	if (Delegates->AreTransformsLocked.Execute())
+	{
+		Delegates->MultiplyScaleToAll.Execute(InScale / GetScale());
+	}
+	else
+	{
+		SetScale(InScale);
+	}
+}
+
+void SMediaImageViewerOverlay::SetScale(float InScale)
+{
+	if (TSharedPtr<FMediaImageViewer> ImageViewer = GetImageViewerDelegate.Execute(Position))
+	{
+		ImageViewer->GetPaintSettings().Scale = InScale;
+	}
+}
+
+void SMediaImageViewerOverlay::MultiplyScale(float InMultiple)
+{
+	SetScale(GetScale() * InMultiple);
+}
+
+void SMediaImageViewerOverlay::ScaleToFit()
+{
+	TSharedPtr<FMediaImageViewer> ImageViewer = GetImageViewerDelegate.Execute(Position);
+
+	if (!ImageViewer.IsValid())
+	{
+		return;
+	}
+
+	const FVector2D LocalSize = Delegates->GetViewerSize.Execute();
+	const FIntPoint& ImageSize = ImageViewer->GetInfo().Size;
+
+	const float ScaleX = LocalSize.X / static_cast<float>(ImageSize.X);
+	const float ScaleY = LocalSize.Y / static_cast<float>(ImageSize.Y);
+
+	if (ScaleX < ScaleY)
+	{
+		Try_SetScale(ScaleX);
+	}
+	else
+	{
+		Try_SetScale(ScaleY);
+	}
+}
+
+bool SMediaImageViewerOverlay::CanResetTransform()
+{
+	return Delegates->GetActiveViews.Execute() == EMediaImageViewerActivePosition::Both;
+}
+
+void SMediaImageViewerOverlay::Try_ResetTransform()
+{
+	if (Delegates->AreTransformsLocked.Execute())
+	{
+		Delegates->ResetTransformToAll.Execute();
+	}
+	else
+	{
+		ResetTransform();
+	}
+}
+
+void SMediaImageViewerOverlay::Try_SetTransform(FVector InOffset, FRotator InRotation, float InScale)
+{
+	if (Delegates->AreTransformsLocked.Execute())
+	{
+		Delegates->SetTransformToAll.Execute(InOffset, InRotation, InScale);
+	}
+	else
+	{
+		SetOffset(InOffset);
+		SetRotation(InRotation);
+		SetScale(InScale);
+	}
+}
+
+bool SMediaImageViewerOverlay::CanCopyTransform()
+{
+	return Delegates->GetActiveViews.Execute() == EMediaImageViewerActivePosition::Both;
+}
+
+void SMediaImageViewerOverlay::CopyTransform()
+{
+	Delegates->SetTransformToAll.Execute(GetOffset(), GetRotation(), GetScale());
+}
+
+uint8 SMediaImageViewerOverlay::GetMipLevel() const
+{
+	if (TSharedPtr<FMediaImageViewer> ImageViewer = GetImageViewerDelegate.Execute(Position))
+	{
+		return ImageViewer->GetPaintSettings().MipLevel;
+	}
+
+	return 0;
+}
+
+void SMediaImageViewerOverlay::AdjustMipLevel(int8 InAdjustment)
+{
+	const uint8 CurrentMipLevel = GetMipLevel();
+
+	if (-InAdjustment > CurrentMipLevel)
+	{
+		SetMipLevel(0);
+	}
+	else
+	{
+		SetMipLevel(CurrentMipLevel + InAdjustment);
+	}
+}
+
+void SMediaImageViewerOverlay::SetMipLevel(uint8 InMipLevel)
+{
+	if (TSharedPtr<FMediaImageViewer> ImageViewer = GetImageViewerDelegate.Execute(Position))
+	{
+		if (InMipLevel < ImageViewer->GetInfo().MipCount)
+		{
+			ImageViewer->GetPaintSettings().MipLevel = InMipLevel;
+		}
+	}
+}
+
+void SMediaImageViewerOverlay::CopyColor()
+{
+	TSharedPtr<FMediaImageViewer> ImageViewer = GetImageViewerDelegate.Execute(Position);
+
+	if (!ImageViewer.IsValid() || !ImageViewer->IsValid())
+	{
+		return;
+	}
+
+	const FIntPoint PixelCoordinates = Delegates->GetPixelCoordinates.Execute(Position);
+
+	if (PixelCoordinates.X < 0 || PixelCoordinates.Y < 0)
+	{
+		return;
+	}
+
+	const FIntPoint& ImageSize = ImageViewer->GetInfo().Size;
+
+	if (PixelCoordinates.X >= ImageSize.X || PixelCoordinates.Y >= ImageSize.Y)
+	{
+		return;
+	}
+
+	const TOptional<TVariant<FColor, FLinearColor>> PixelColor = ImageViewer->GetPixelColor(PixelCoordinates, ImageViewer->GetPaintSettings().MipLevel);
+
+	if (!PixelColor.IsSet())
+	{
+		return;
+	}
+
+	const TVariant<FColor, FLinearColor>& PixelColorValue = PixelColor.GetValue();
+
+	if (const FColor* Color = PixelColorValue.TryGet<FColor>())
+	{
+		FPlatformApplicationMisc::ClipboardCopy(*Color->ToString());
+	}
+
+	if (const FLinearColor* ColorLinear = PixelColorValue.TryGet<FLinearColor>())
+	{
+		FPlatformApplicationMisc::ClipboardCopy(*ColorLinear->ToString());
+	}
+}
+
+bool SMediaImageViewerOverlay::CanAddToLibrary() const
+{
+	TSharedRef<FMediaViewerLibrary> Library = Delegates->GetLibrary.Execute();
+
+	return !Library->FindItemByValue(CachedItem->GetItemType(), CachedItem->GetStringValue()).IsValid();
+}
+
+void SMediaImageViewerOverlay::AddToLibrary()
+{
+	TSharedRef<FMediaViewerLibrary> Library = Delegates->GetLibrary.Execute();
+
+	Library->AddItemToGroup(CachedItem.ToSharedRef());
+}
+
+TSharedRef<SWidget> SMediaImageViewerOverlay::CreateMenu()
+{
+	FMenuBuilder MenuBuilder(true, CommandList, nullptr, false, &FAppStyle::Get(), false);
+
+	MenuBuilder.BeginSection(TEXT("Image"));
+	{
+		MenuBuilder.AddMenuEntry(FMediaViewerCommands::Get().ResetTransform);
+		MenuBuilder.AddMenuEntry(FMediaViewerCommands::Get().CopyTransform);
+		MenuBuilder.AddMenuEntry(FMediaViewerCommands::Get().AddToLibrary);
+	}
+	MenuBuilder.EndSection();
+
+	MenuBuilder.AddMenuSeparator();
+
+	MenuBuilder.BeginSection(TEXT("Media Viewer"));
+	{
+		MenuBuilder.AddMenuEntry(FMediaViewerCommands::Get().ToggleOverlay);
+		MenuBuilder.AddMenuEntry(FMediaViewerCommands::Get().ToggleLockedTransform);
+		MenuBuilder.AddMenuEntry(FMediaViewerCommands::Get().ResetAllTransforms);
+		MenuBuilder.AddMenuEntry(FMediaViewerCommands::Get().SwapAB);
+	}
+
+	return MenuBuilder.MakeWidget();
+}
+
+} // UE::MediaViewer::Private
+
+#undef LOCTEXT_NAMESPACE
diff --git a/Engine/Plugins/Experimental/MediaViewer/Source/MediaViewer/Private/Widgets/SMediaImageViewerOverlay.h b/Engine/Plugins/Experimental/MediaViewer/Source/MediaViewer/Private/Widgets/SMediaImageViewerOverlay.h
new file mode 100644
--- /dev/null
+++ b/Engine/Plugins/Experimental/MediaViewer/Source/MediaViewer/Private/Widgets/SMediaImageViewerOverlay.h
@@ -0,0 +1,131 @@
+// Copyright Epic Games, Inc. All Rights Reserved.
+
+#pragma once
+
+#include "Widgets/SCompoundWidget.h"
+
+#include "Widgets/MediaViewerDelegates.h"
+
+class FUICommandList;
+class SVerticalBox;
+struct FMediaViewerLibraryItem;
+
+namespace UE::MediaViewer
+{
+	class FMediaImageViewer;
+	enum class EMediaImageViewerPosition : uint8;
+}
+
+namespace UE::MediaViewer::Private
+{
+
+class SMediaViewerLibrary;
+class SMediaViewerLibraryPrivate;
+
+class SMediaImageViewerOverlay : public SCompoundWidget
+{
+	SLATE_DECLARE_WIDGET(SMediaImageViewerOverlay, SCompoundWidget)
+
+	SLATE_BEGIN_ARGS(SMediaImageViewerOverlay)
+		: _bCreateSplitDropTarget(false)
+		{}
+		SLATE_ARGUMENT(bool, bCreateSplitDropTarget)
+	SLATE_END_ARGS()
+
+public:
+	SMediaImageViewerOverlay();
+
+	virtual ~SMediaImageViewerOverlay() override = default;
+
+	void Construct(const FArguments& InArgs, EMediaImageViewerPosition InPosition, const TSharedPtr<FMediaViewerDelegates>& InDelegates);
+
+	TSharedPtr<FMediaImageViewer> GetImageViewer() const;
+
+	const TSharedRef<FUICommandList>& GetCommandList() const;
+
+	FIntPoint GetImageViewerPixelCoordinates() const;
+
+	bool IsCursorOverImageViewer() const;
+
+	void UpdateMouse(const TOptional<FIntPoint>& InMousePosition);
+
+	bool IsOverlayEnabled() const;
+	void ToggleOverlay();
+
+	//~ Begin SWidget
+	virtual FCursorReply OnCursorQuery(const FGeometry& InMyGeometry, const FPointerEvent& InCursorEvent) const override;
+	virtual FReply OnKeyDown(const FGeometry& InMyGeometry, const FKeyEvent& InKeyEvent) override;
+	virtual FReply OnMouseButtonDown(const FGeometry& InMyGeometry, const FPointerEvent& InMouseEvent) override;
+	virtual FReply OnMouseButtonUp(const FGeometry& InMyGeometry, const FPointerEvent& InMouseEvent) override;
+	virtual FReply OnMouseWheel(const FGeometry& InMyGeometry, const FPointerEvent& InMouseEvent) override;
+	virtual bool SupportsKeyboardFocus() const override;
+	//~ End SWidget
+
+protected:
+	EMediaImageViewerPosition Position;
+	FMediaViewerDelegates::FGetImageViewer GetImageViewerDelegate;
+	TSharedPtr<FMediaViewerLibraryItem> CachedItem;
+	TSharedRef<FUICommandList> CommandList;
+	FName DragButtonName;
+	bool bDragging;
+	bool bExternalDragging;
+	FVector2D DragStartCursor;
+	FVector DragStartOffset;
+	FIntPoint LastCursorPosition;
+	TSharedPtr<FMediaViewerDelegates> Delegates;
+	TSharedPtr<SVerticalBox> Container;
+	bool bOverlayEnabled;
+
+	void BindCommands();
+
+	TSharedRef<SWidget> CreateOverlay(bool bInCreateSplitDropTarget);
+
+	TSharedRef<SWidget> CreateStatusBar(const TSharedPtr<FMediaViewerDelegates>& InDelegates);
+
+	void OnDragButtonUp(FName InKeyName);
+	void UpdateDragPosition();
+
+	bool OnAssetDraggedOver(TArrayView<FAssetData> InAssetData, FText& OutError) const;
+	void OnAssetsDropped(const FDragDropEvent& InDragDropEvent, TArrayView<FAssetData> InAssetData, EMediaImageViewerPosition InPosition);
+
+	bool OnVerifyDrag(TSharedPtr<FDragDropOperation> InDragDropOp);
+	FReply OnDrop(const FGeometry& InMyGeometry, const FDragDropEvent& InDragDropEvent);
+
+	EVisibility GetDragDescriptionVisibility() const;
+
+	FVector GetOffset() const;
+	void Try_AddOffset(FVector InOffset);
+	void AddOffset(const FVector& InOffset);
+	void SetOffset(const FVector& InOffset);
+
+	FRotator GetRotation() const;
+	void Try_AddRotation(FRotator InRotation);
+	void AddRotation(const FRotator& InRotation);
+	void SetRotation(const FRotator& InRotation);
+
+	float GetScale() const;
+	void Try_SetScale(float InScale);
+	void SetScale(float InScale);
+	void MultiplyScale(float InMultiple);
+	void ScaleToFit();
+
+	bool CanResetTransform();
+	void ResetTransform();
+	void Try_ResetTransform();
+	void Try_SetTransform(FVector InOffset, FRotator InRotation, float InScale);
+	bool CanCopyTransform();
+	void CopyTransform();
+
+	uint8 GetMipLevel() const;
+	void AdjustMipLevel(int8 InAdjustment);
+	void SetMipLevel(uint8 InMipLevel);
+
+	void CopyColor();
+
+	bool CanAddToLibrary() const;
+	void AddToLibrary();
+
+	TSharedRef<SWidget> CreateMenu();
+};
+
+} // UE::MediaViewer::Private
diff --git a/Engine/Plugins/Experimental/MediaViewer/Source/MediaViewer/Private/Widgets/SMediaImageViewerStatusBar.h b/Engine/Plugins/Experimental/MediaViewer/Source/MediaViewer/Private/Widgets/SMediaImageViewerStatusBar.h
new file mode 100644
--- /dev/null
+++ b/Engine/Plugins/Experimental/MediaViewer/Source/MediaViewer/Private/Widgets/SMediaImageViewerStatusBar.h
@@ -0,0 +1,47 @@
+// Copyright Epic Games, Inc. All Rights Reserved.
+
+#pragma once
+
+#include "Widgets/SCompoundWidget.h"
+
+#include "Widgets/MediaViewerDelegates.h"
+
+class FText;
+class FUICommandList;
+
+namespace UE::MediaViewer
+{
+	class FMediaImageViewer;
+}
+
+namespace UE::MediaViewer::Private
+{
+
+class SMediaImageViewerStatusBar : public SCompoundWidget
+{
+	SLATE_DECLARE_WIDGET(SMediaImageViewerStatusBar, SCompoundWidget)
+
+	SLATE_BEGIN_ARGS(SMediaImageViewerStatusBar)
+		{}
+	SLATE_END_ARGS()
+
+public:
+	SMediaImageViewerStatusBar();
+
+	virtual ~SMediaImageViewerStatusBar() override = default;
+
+	void Construct(const FArguments& InArgs, EMediaImageViewerPosition InPosition, const TSharedPtr<FMediaViewerDelegates>& InDelegates);
+
+protected:
+	EMediaImageViewerPosition Position;
+	FMediaViewerDelegates::FGetImageViewer GetImageViewerDelegate;
+	FMediaViewerDelegates::FGetLocationForPosition PixelLocationDelegate;
+
+	TSharedRef<SWidget> BuildStatusBar();
+
+	FText GetResolutionLabel() const;
+
+	FText GetColorPickerLabel() const;
+};
+
+} // UE::MediaViewer::Private
diff --git a/Engine/Plugins/Experimental/MediaViewer/Source/MediaViewer/Private/Widgets/SMediaViewerLibraryPrivate.cpp b/Engine/Plugins/Experimental/MediaViewer/Source/MediaViewer/Private/Widgets/SMediaViewerLibraryPrivate.cpp
new file mode 100644
--- /dev/null
+++ b/Engine/Plugins/Experimental/MediaViewer/Source/MediaViewer/Private/Widgets/SMediaViewerLibraryPrivate.cpp
@@ -0,0 +1,157 @@
+// Copyright Epic Games, Inc. All Rights Reserved.
+
+#include "Widgets/SMediaViewerLibraryPrivate.h"
+
+#include "ImageViewers/ColorImageViewer.h"
+#include "IMediaViewerModule.h"
+#include "LevelEditor.h"
+#include "Library/MediaViewerLibrary.h"
+#include "Library/MediaViewerLibraryGroup.h"
+#include "Library/MediaViewerLibraryIni.h"
+#include "Library/MediaViewerLibraryItem.h"
+#include "Library/MediaViewerLibraryMediaTextureGroup.h"
+#include "Library/MediaViewerLibraryTextureRenderTarget2DGroup.h"
+#include "MediaViewer.h"
+#include "Misc/Optional.h"
+#include "Widgets/MediaViewerDelegates.h"
+#include "Widgets/SMediaViewer.h"
+#include "Widgets/SMediaViewerLibrary.h"
+
+#define LOCTEXT_NAMESPACE "SMediaViewerLibraryPrivate"
+
+namespace UE::MediaViewer::Private
+{
+	constexpr int32 MaxHistoryEntries = 25;
+	static const FGuid ColorGroupId = FGuid(0x27FDA3CE, 0x575E24CB, 0xE3717531, 0xCA9D2169);
+}
+
+namespace UE::MediaViewer::Private
+{
+
+SMediaViewerLibraryPrivate::SMediaViewerLibraryPrivate()
+{
+}
+
+void SMediaViewerLibraryPrivate::PrivateRegisterAttributes(FSlateAttributeDescriptor::FInitializer&)
+{
+}
+
+void SMediaViewerLibraryPrivate::Construct(const FArguments& InArgs, const TSharedPtr<FMediaViewerDelegates>& InDelegates)
+{
+	IMediaViewerLibraryWidget::FArgs Args;
+
+	ChildSlot
+	[
+		SAssignNew(Library, SMediaViewerLibrary, Args)
+	];
+
+	if (!::UMediaViewerLibraryIni::Get().HasGroup(UE::MediaViewer::Private::ColorGroupId))
+	{
+		AddColors();
+	}
+
+	TSharedRef<FMediaViewerLibrary> LibraryImpl = GetLibrary();
+
+	LibraryImpl->AddGroup(MakeShared<FMediaViewerLibraryMediaTextureGroup>(LibraryImpl));
+	LibraryImpl->AddGroup(MakeShared<FMediaViewerLibraryTextureRenderTarget2DGroup>(LibraryImpl));
+}
+
+TSharedRef<FMediaViewerLibrary> SMediaViewerLibraryPrivate::GetLibrary() const
+{
+	return Library->GetLibraryImpl();
+}
+
+void SMediaViewerLibraryPrivate::OnImageViewerOpened(const TSharedRef<FMediaImageViewer>& InImageViewer)
+{
+	TSharedRef<IMediaViewerLibrary> LibraryImpl = Library->GetLibrary();
+
+	TSharedPtr<FMediaViewerLibraryGroup> HistoryGroup = LibraryImpl->GetGroup(LibraryImpl->GetHistoryGroupId());
+
+	if (!HistoryGroup.IsValid())
+	{
+		return;
+	}
+
+	TSharedPtr<FMediaViewerLibraryItem> Item = LibraryImpl->GetItem(InImageViewer->GetInfo().Id);
+
+	if (!Item.IsValid())
+	{
+		Item = InImageViewer->CreateLibraryItem();
+
+		if (!Item.IsValid())
+		{
+			return;
+		}
+
+		LibraryImpl->AddItem(Item.ToSharedRef());
+	}
+
+	HistoryGroup->RemoveItem(Item->GetId());
+
+	HistoryGroup->AddItem(Item->GetId(), 0);
+
+	while (HistoryGroup->GetItems().Num() > UE::MediaViewer::Private::MaxHistoryEntries)
+	{
+		HistoryGroup->RemoveItem(HistoryGroup->GetItems().Last());
+	}
+
+	Library->OnLibraryChanged(Library->GetLibrary(), IMediaViewerLibrary::EChangeType::ItemGroupChanged);
+}
+
+void SMediaViewerLibraryPrivate::AddColors()
+{
+	TSharedRef<IMediaViewerLibrary> LibraryImpl = Library->GetLibrary();
+	TSharedPtr<FMediaViewerLibraryGroup> ColorGroup = LibraryImpl->GetGroup(UE::MediaViewer::Private::ColorGroupId);
+
+	if (!ColorGroup.IsValid())
+	{
+		ColorGroup = MakeShared<FMediaViewerLibraryGroup>(
+			UE::MediaViewer::Private::ColorGroupId,
+			LOCTEXT("Colors", "Colors"),
+			LOCTEXT("ColorsToolTip", "A set of colors"),
+			/* Dynamic */ false
+		);
+
+		LibraryImpl->AddGroup(ColorGroup.ToSharedRef());
+	}
+	
+	static TSharedRef<FColorImageViewer::FItem> ColorItems[5] = {
+		MakeShared<FColorImageViewer::FItem>(FGuid::Combine(UE::MediaViewer::Private::ColorGroupId, FGuid(0, 0, 0, 1)),
+			INVTEXT("Red"), INVTEXT("It's Red!"), FLinearColor::Red),
+
+		MakeShared<FColorImageViewer::FItem>(FGuid::Combine(UE::MediaViewer::Private::ColorGroupId, FGuid(0, 0, 0, 2)),
+			INVTEXT("Green"), INVTEXT("It's Green!"), FLinearColor::Green),
+
+		MakeShared<FColorImageViewer::FItem>(FGuid::Combine(UE::MediaViewer::Private::ColorGroupId, FGuid(0, 0, 0, 3)),
+			INVTEXT("Blue"), INVTEXT("It's Blue!"), FLinearColor::Blue),
+
+		MakeShared<FColorImageViewer::FItem>(FGuid::Combine(UE::MediaViewer::Private::ColorGroupId, FGuid(0, 0, 0, 4)),
+			INVTEXT("White"), INVTEXT("It's White!"), FLinearColor::White),
+
+		MakeShared<FColorImageViewer::FItem>(FGuid::Combine(UE::MediaViewer::Private::ColorGroupId, FGuid(0, 0, 0, 5)),
+			INVTEXT("Black"), INVTEXT("It's Black!"), FLinearColor::Black)
+	};
+
+	while (!ColorGroup->GetItems().IsEmpty())
+	{
+		ColorGroup->RemoveItem(ColorGroup->GetItems()[0]);
+	}
+
+	ColorGroup->AddItem(ColorItems[0]->GetId());
+	ColorGroup->AddItem(ColorItems[1]->GetId());
+	ColorGroup->AddItem(ColorItems[2]->GetId());
+	ColorGroup->AddItem(ColorItems[3]->GetId());
+	ColorGroup->AddItem(ColorItems[4]->GetId());
+
+	for (const TSharedRef<FColorImageViewer::FItem>& ColorItem : ColorItems)
+	{
+		if (!LibraryImpl->GetItem(ColorItem->GetId()))
+		{
+			LibraryImpl->AddItem(ColorItem);
+		}
+	}
+}
+
+} // UE::MediaViewer::Private
+
+#undef LOCTEXT_NAMESPACE
diff --git a/Engine/Plugins/Experimental/MediaViewer/Source/MediaViewer/Private/Widgets/SMediaViewerToolbar.cpp b/Engine/Plugins/Experimental/MediaViewer/Source/MediaViewer/Private/Widgets/SMediaViewerToolbar.cpp
new file mode 100644
--- /dev/null
+++ b/Engine/Plugins/Experimental/MediaViewer/Source/MediaViewer/Private/Widgets/SMediaViewerToolbar.cpp
@@ -0,0 +1,470 @@
+// Copyright Epic Games, Inc. All Rights Reserved.
+
+#include "Widgets/SMediaViewerToolbar.h"
+
+#include "DetailLayoutBuilder.h"
+#include "Framework/MultiBox/MultiBoxBuilder.h"
+#include "Framework/MultiBox/SToolBarButtonBlock.h"
+#include "ImageViewer/MediaImageViewer.h"
+#include "ImageViewers/NullImageViewer.h"
+#include "IMediaViewerModule.h"
+#include "Library/MediaViewerLibrary.h"
+#include "MediaViewer.h"
+#include "MediaViewerCommands.h"
+#include "MediaViewerStyle.h"
+#include "MediaViewerUtils.h"
+#include "SAssetDropTarget.h"
+#include "SEditorViewportToolBarMenu.h"
+#include "Styling/SlateTypes.h"
+#include "Styling/ToolBarStyle.h"
+#include "UObject/StructOnScope.h"
+#include "Widgets/Colors/SColorBlock.h"
+#include "Widgets/SMediaImageViewerDetails.h"
+#include "Widgets/SMediaViewer.h"
+#include "Widgets/Text/STextBlock.h"
+
+#define LOCTEXT_NAMESPACE "SMediaViewerToolbar"
+
+namespace UE::MediaViewer::Private
+{
+
+SMediaViewerToolbar::SMediaViewerToolbar()
+{
+}
+
+void SMediaViewerToolbar::PrivateRegisterAttributes(FSlateAttributeDescriptor::FInitializer&)
+{
+}
+
+void SMediaViewerToolbar::Construct(const FArguments& InArgs, const TSharedPtr<FMediaViewerDelegates>& InDelegates)
+{
+	using namespace UE::MediaViewer;
+
+	Delegates = InDelegates;
+
+	for (int32 Index = 0; Index < static_cast<int32>(EMediaImageViewerPosition::COUNT); ++Index)
+	{
+		ImageDetails[Index] = SNullWidget::NullWidget;
+	}
+
+	MediaViewerSettingsView = FMediaViewerUtils::CreateStructDetailsView(
+		MakeShared<FStructOnScope>(FMediaViewerSettings::StaticStruct(), reinterpret_cast<uint8*>(InDelegates->GetSettings.Execute())),
+		LOCTEXT("BackgroundTexture", "Background Texture"),
+		this
+	);
+
+	MediaViewerSettingsWidget = SNew(SBox)
+		.Padding(3.f)
+		[
+			MediaViewerSettingsView->GetWidget().ToSharedRef()
+		];
+
+	TSharedPtr<FMediaImageViewer> FirstImageViewer = InDelegates->GetImageViewer.Execute(EMediaImageViewerPosition::First);
+	const bool bHasFirstImageViewer = FirstImageViewer.IsValid() && FirstImageViewer != FNullImageViewer::GetNullImageViewer();
+
+	TSharedPtr<FMediaImageViewer> SecondImageViewer = InDelegates->GetImageViewer.Execute(EMediaImageViewerPosition::Second);
+	const bool bHasSecondImageViewer = SecondImageViewer.IsValid() && SecondImageViewer != FNullImageViewer::GetNullImageViewer();
+
+	const float CenterPaddingLeft = bHasFirstImageViewer ? 0.f : 83.f;
+	const float CenterPaddingRight = bHasSecondImageViewer ? 0.f : 83.f;
+
+	ChildSlot
+	[
+		SNew(SBorder)
+		.HAlign(HAlign_Fill)
+		.BorderImage(FAppStyle::GetBrush("EditorViewportToolBar.Background"))
+		[
+			SNew(SHorizontalBox)
+			+ SHorizontalBox::Slot()
+			.AutoWidth()
+			.HAlign(HAlign_Left)
+			.VAlign(VAlign_Center)
+			[
+				FirstImageViewer
+					? MakeSideToolbar(EMediaImageViewerPosition::First, ToolbarSections::ToolbarLeft)
+					: SNullWidget::NullWidget
+			]
+			+ SHorizontalBox::Slot()
+			.FillWidth(1.f)
+			.HAlign(HAlign_Center)
+			.VAlign(VAlign_Center)
+			.Padding(CenterPaddingLeft, 0.f, CenterPaddingRight, 0.f)
+			[
+				MakeCenterToolbar()
+			]
+			+ SHorizontalBox::Slot()
+			.AutoWidth()
+			.HAlign(HAlign_Right)
+			.VAlign(VAlign_Center)
+			[				
+				bHasSecondImageViewer
+					? MakeSideToolbar(EMediaImageViewerPosition::Second, ToolbarSections::ToolbarRight)
+					: SNullWidget::NullWidget
+			]
+		]
+	];
+}
+
+void SMediaViewerToolbar::NotifyPostChange(const FPropertyChangedEvent& InPropertyChangedEvent, FProperty* InPropertyThatChanged)
+{
+	const FName MemberPropertyName = InPropertyChangedEvent.GetMemberPropertyName();
+
+	if (MemberPropertyName == GET_MEMBER_NAME_CHECKED(FMediaViewerSettings, ABOrientation))
+	{
+		Delegates->RefreshView.Execute();
+	}
+	else if (MemberPropertyName == GET_MEMBER_NAME_CHECKED(FMediaViewerSettings, ABSplitterLocation))
+	{
+		Delegates->SetABSplitterLocation.Execute(Delegates->GetABSplitterLocation.Execute());
+	}
+}
+
+TSharedRef<SWidget> SMediaViewerToolbar::MakeCenterToolbar()
+{
+	using namespace UE::MediaViewer;
+
+	FSlimHorizontalToolBarBuilder ToolbarBuilder(Delegates->GetCommandList.Execute(), FMultiBoxCustomization::None);
+	ToolbarBuilder.SetStyle(&FAppStyle::Get(), "ViewportToolbar");
+	ToolbarBuilder.SetIsFocusable(false);
+
+	EMediaImageViewerActivePosition ActiveViews = Delegates->GetActiveViews.Execute();
+
+	const FMediaViewerCommands& Commands = FMediaViewerCommands::Get();
+
+	FSlateIcon ViewIcon;
+
+	if (Delegates->GetActiveViews.Execute() != EMediaImageViewerActivePosition::Both)
+	{
+		ViewIcon = FSlateIcon(FAppStyle::Get().GetStyleSetName(), "EditorViewport.ViewportConfig_OnePane");
+	}
+	else if (Delegates->GetABOrientation.Execute() == Orient_Horizontal)
+	{
+		ViewIcon = FSlateIcon(FAppStyle::Get().GetStyleSetName(), "EditorViewport.ViewportConfig_TwoPanesH");
+	}
+	else
+	{
+		ViewIcon = FSlateIcon(FAppStyle::Get().GetStyleSetName(), "EditorViewport.ViewportConfig_TwoPanesV");
+	}
+
+	ToolbarBuilder.BeginSection(ToolbarSections::ToolbarCenter);
+	{
+		ToolbarBuilder.BeginBlockGroup();
+		{
+			ToolbarBuilder.AddComboButton(
+				FUIAction(),
+				FOnGetContent::CreateSP(this, &SMediaViewerToolbar::MakeViewMenu),
+				FText::GetEmpty(),
+				LOCTEXT("ViewType", "Change the type of view between single, AB horizontal and AB vertical"),
+				ViewIcon
+			);
+		}
+		ToolbarBuilder.EndBlockGroup();
+
+		ToolbarBuilder.BeginBlockGroup();
+		{
+			ToolbarBuilder.AddToolBarButton(
+				Commands.SwapAB,
+				NAME_None,
+				FText::GetEmpty(),
+				LOCTEXT("SwapABTooltip", "Swaps the image viewers and their offsets."),
+				FSlateIcon(FAppStyle::Get().GetStyleSetName(), "MeshPaint.SwapColor")
+			);
+
+			ToolbarBuilder.AddToolBarButton(
+				Commands.ToggleLockedTransform,
+				NAME_None,
+				FText::GetEmpty(),
+				LOCTEXT("ToggleLockedTransformToolTip", "Toggle transform lock between viewers."),
+				FSlateIcon(FAppStyle::Get().GetStyleSetName(), "GenericLink")
+			);
+
+			ToolbarBuilder.AddToolBarButton(
+				Commands.ResetAllTransforms,
+				NAME_None,
+				FText::GetEmpty(),
+				LOCTEXT("ResetTransformToAlToolTip", "Reset the camera transform for all viewers."),
+				FSlateIcon(FAppStyle::Get().GetStyleSetName(), "BlueprintEditor.ResetCamera")
+			);
+		}
+		ToolbarBuilder.EndBlockGroup();
+
+		ToolbarBuilder.BeginBlockGroup();
+		{
+			ToolbarBuilder.AddComboButton(
+				FUIAction(),
+				FOnGetContent::CreateSP(this, &SMediaViewerToolbar::GetBackgroundTextureSettingsWidget),
+				FText::GetEmpty(),
+				FText::GetEmpty(),
+				FSlateIcon(FAppStyle::Get().GetStyleSetName(), "Icons.Toolbar.Settings")
+			);
+		}
+		ToolbarBuilder.EndBlockGroup();
+	}
+	ToolbarBuilder.EndSection();
+
+	return ToolbarBuilder.MakeWidget();
+}
+
+TSharedRef<SWidget> SMediaViewerToolbar::MakeSideToolbar(EMediaImageViewerPosition InPosition, FName InToolbarName)
+{
+	using namespace UE::MediaViewer;
+
+	TSharedPtr<FMediaImageViewer> ImageViewer = Delegates->GetImageViewer.Execute(InPosition);
+	const bool bValidImageViewer = ImageViewer.IsValid() && ImageViewer != FNullImageViewer::GetNullImageViewer();
+
+	TSharedRef<FExtender> Extender = MakeShared<FExtender>();
+
+	if (bValidImageViewer)
+	{
+		ImageViewer->ExtendToolbar(Extender);
+		ImageDetails[static_cast<int32>(InPosition)] = SNew(SBox)
+			.Padding(3.f)
+			[
+				SNew(SMediaImageViewerDetails, InPosition, Delegates)
+			];
+	}
+
+	FSlimHorizontalToolBarBuilder ToolbarBuilder(Delegates->GetCommandListForPosition.Execute(InPosition), FMultiBoxCustomization::None, Extender);
+	ToolbarBuilder.SetStyle(&FAppStyle::Get(), "ViewportToolbar");
+	ToolbarBuilder.SetIsFocusable(false);
+
+	if (!bValidImageViewer)
+	{
+		return ToolbarBuilder.MakeWidget();
+	}
+
+	ToolbarBuilder.BeginSection(InToolbarName);
+	{
+		if (ImageViewer.IsValid())
+		{
+			ToolbarBuilder.BeginBlockGroup();
+
+			{
+				ToolbarBuilder.AddComboButton(
+					FUIAction(),
+					FOnGetContent::CreateSP(this, &SMediaViewerToolbar::GetDetailsWidget, InPosition),
+					FText::GetEmpty(),
+					FText::GetEmpty(),
+					FSlateIcon(FAppStyle::Get().GetStyleSetName(), "LevelEditor.Tabs.Details")
+				);
+
+				ToolbarBuilder.AddComboButton(
+					FUIAction(),
+					FOnGetContent::CreateSP(this, &SMediaViewerToolbar::MakeScaleMenu, InPosition),
+					TAttribute<FText>::CreateSP(this, &SMediaViewerToolbar::GetScaleMenuLabel, InPosition),
+					FText::GetEmpty(),
+					FSlateIcon(FAppStyle::GetAppStyleSetName(), "EditorViewport.ScaleGridSnap")
+				);
+			}
+
+			ToolbarBuilder.EndBlockGroup();
+		}
+	}
+	ToolbarBuilder.EndSection();
+
+	return ToolbarBuilder.MakeWidget();
+}
+
+bool SMediaViewerToolbar::OnAssetDraggedOver(TArrayView<FAssetData> InAssetData) const
+{
+	IMediaViewerModule& Module = IMediaViewerModule::Get();
+
+	for (const FAssetData& AssetData : InAssetData)
+	{
+		if (Module.HasFactoryFor(AssetData))
+		{
+			return true;
+		}
+	}
+
+	return false;
+}
+
+void SMediaViewerToolbar::OnAssetsDropped(const FDragDropEvent& InDragDropEvent, TArrayView<FAssetData> InAssetData)
+{
+	IMediaViewerModule& Module = IMediaViewerModule::Get();
+
+	int32 AssetIndex = INDEX_NONE;
+
+	for (int32 Index = 0; Index < InAssetData.Num(); ++Index)
+	{
+		if (Module.HasFactoryFor(InAssetData[Index]))
+		{
+			AssetIndex = Index;
+			break;
+		}
+	}
+
+	if (AssetIndex == INDEX_NONE)
+	{
+		return;
+	}
+
+	if (!Module.HasImage(EMediaImageViewerPosition::First))
+	{
+		Module.SetImage(EMediaImageViewerPosition::First, InAssetData[AssetIndex]);
+	}
+	else if (!Module.HasImage(EMediaImageViewerPosition::Second))
+	{
+		Module.SetImage(EMediaImageViewerPosition::Second, InAssetData[AssetIndex]);
+	}
+}
+
+FText SMediaViewerToolbar::GetScaleMenuLabel(EMediaImageViewerPosition InPosition) const
+{
+	if (TSharedPtr<FMediaImageViewer> ImageViewer = Delegates->GetImageViewer.Execute(InPosition))
+	{
+		return FText::Format(
+			LOCTEXT("Scale", "{0}%"),
+			FText::AsNumber(ImageViewer->GetPaintSettings().Scale * 100.f)
+		);
+	}
+
+	return INVTEXT("-");
+}
+
+TSharedRef<SWidget> SMediaViewerToolbar::MakeScaleMenu(EMediaImageViewerPosition InPosition) const
+{
+	TSharedPtr<FUICommandList> CommandList = Delegates->GetCommandListForPosition.Execute(InPosition);
+
+	FMenuBuilder MenuBuilder(true, CommandList, nullptr, false, &FAppStyle::Get(), false);
+
+	if (CommandList.IsValid())
+	{
+		const FMediaViewerCommands& Commands = FMediaViewerCommands::Get();
+
+		MenuBuilder.AddMenuEntry(Commands.Scale12);
+		MenuBuilder.AddMenuEntry(Commands.Scale25);
+		MenuBuilder.AddMenuEntry(Commands.Scale50);
+		MenuBuilder.AddMenuEntry(Commands.Scale100);
+		MenuBuilder.AddMenuEntry(Commands.Scale200);
+		MenuBuilder.AddMenuEntry(Commands.Scale400);
+		MenuBuilder.AddMenuEntry(Commands.Scale800);
+		MenuBuilder.AddSeparator();
+		MenuBuilder.AddMenuEntry(Commands.ScaleToFit);
+	}
+
+	return MenuBuilder.MakeWidget();
+}
+
+TSharedRef<SWidget> SMediaViewerToolbar::GetDetailsWidget(EMediaImageViewerPosition InPosition) const
+{
+	return ImageDetails[static_cast<int32>(InPosition)].ToSharedRef();
+}
+
+FReply SMediaViewerToolbar::OnSaveToLibraryClicked()
+{
+	return FReply::Handled();
+}
+
+TSharedRef<SWidget> SMediaViewerToolbar::GetBackgroundTextureSettingsWidget() const
+{
+	return MediaViewerSettingsWidget.ToSharedRef();
+}
+
+TSharedRef<SWidget> SMediaViewerToolbar::MakeViewMenu()
+{
+	FMenuBuilder MenuBuilder(true, nullptr, nullptr, false, &FAppStyle::Get(), false);
+
+	FMenuEntryParams SingleViewParams;
+	SingleViewParams.LabelOverride = LOCTEXT("SetSingleView", "Single Image");
+	SingleViewParams.ToolTipOverride = LOCTEXT("SetSingleViewToolTip", "View a single image.");
+	SingleViewParams.IconOverride = FSlateIcon(FAppStyle::Get().GetStyleSetName(), "EditorViewport.ViewportConfig_OnePane");
+	SingleViewParams.UserInterfaceActionType = EUserInterfaceActionType::RadioButton;
+
+	SingleViewParams.DirectActions.ExecuteAction.BindSPLambda(
+		this,
+		[this]
+		{
+			Delegates->SetSingleView.Execute();
+		}
+	);
+
+	SingleViewParams.DirectActions.GetActionCheckState.BindSPLambda(
+		this,
+		[this]
+		{
+			return Delegates->GetActiveViews.Execute() != EMediaImageViewerActivePosition::Both
+				? ECheckBoxState::Checked
+				: ECheckBoxState::Unchecked;
+		}
+	);
+
+	MenuBuilder.AddMenuEntry(SingleViewParams);
+
+	FMenuEntryParams ABHorizontalViewParams;
+	ABHorizontalViewParams.LabelOverride = LOCTEXT("SetABHorizontalView", "Horizontal");
+	ABHorizontalViewParams.ToolTipOverride = LOCTEXT("SetABHorizontalViewToolTip", "View 2 images side by side.");
+	ABHorizontalViewParams.IconOverride = FSlateIcon(FAppStyle::Get().GetStyleSetName(), "EditorViewport.ViewportConfig_TwoPanesH");
+	ABHorizontalViewParams.UserInterfaceActionType = EUserInterfaceActionType::RadioButton;
+
+	ABHorizontalViewParams.DirectActions.ExecuteAction.BindSPLambda(
+		this,
+		[this]
+		{
+			Delegates->SetABOrientation.Execute(Orient_Horizontal);
+		}
+	);
+
+	ABHorizontalViewParams.DirectActions.CanExecuteAction.BindSPLambda(
+		this,
+		[this]
+		{
+			return Delegates->GetActiveViews.Execute() == EMediaImageViewerActivePosition::Both;
+		}
+	);
+
+	ABHorizontalViewParams.DirectActions.GetActionCheckState.BindSPLambda(
+		this,
+		[this]
+		{
+			return (Delegates->GetActiveViews.Execute() == EMediaImageViewerActivePosition::Both
+				&& Delegates->GetABOrientation.Execute() == Orient_Horizontal)
+				? ECheckBoxState::Checked
+				: ECheckBoxState::Unchecked;
+		}
+	);
+
+	MenuBuilder.AddMenuEntry(ABHorizontalViewParams);
+
+	FMenuEntryParams ABVerticalViewParams;
+	ABVerticalViewParams.LabelOverride = LOCTEXT("SetABVerticalView", "Vertical");
+	ABVerticalViewParams.ToolTipOverride = LOCTEXT("SetABVerticalViewToolTip", "View 2 images, one above the other.");
+	ABVerticalViewParams.IconOverride = FSlateIcon(FAppStyle::Get().GetStyleSetName(), "EditorViewport.ViewportConfig_TwoPanesV");
+	ABVerticalViewParams.UserInterfaceActionType = EUserInterfaceActionType::RadioButton;
+
+	ABVerticalViewParams.DirectActions.ExecuteAction.BindSPLambda(
+		this,
+		[this]
+		{
+			Delegates->SetABOrientation.Execute(Orient_Vertical);
+		}
+	);
+
+	ABVerticalViewParams.DirectActions.CanExecuteAction.BindSPLambda(
+		this,
+		[this]
+		{
+			return Delegates->GetActiveViews.Execute() == EMediaImageViewerActivePosition::Both;
+		}
+	);
+
+	ABVerticalViewParams.DirectActions.GetActionCheckState.BindSPLambda(
+		this,
+		[this]
+		{
+			return (Delegates->GetActiveViews.Execute() == EMediaImageViewerActivePosition::Both
+				&& Delegates->GetABOrientation.Execute() == Orient_Vertical)
+				? ECheckBoxState::Checked
+				: ECheckBoxState::Unchecked;
+		}
+	);
+
+	MenuBuilder.AddMenuEntry(ABVerticalViewParams);
+
+	return MenuBuilder.MakeWidget();
+}
+
+} // UE::MediaViewer::Private
+
+#undef LOCTEXT_NAMESPACE
