Media Viewer: Changed default clear color from Panel to Recessed

diff --git a/Engine/Plugins/Experimental/MediaViewer/Source/MediaViewer/Private/Widgets/MediaViewerSettings.h b/Engine/Plugins/Experimental/MediaViewer/Source/MediaViewer/Private/Widgets/MediaViewerSettings.h
--- a/Engine/Plugins/Experimental/MediaViewer/Source/MediaViewer/Private/Widgets/MediaViewerSettings.h
+++ b/Engine/Plugins/Experimental/MediaViewer/Source/MediaViewer/Private/Widgets/MediaViewerSettings.h
@@ -21,7 +21,7 @@ struct FMediaViewerSettings
 
 	/** Color painted to the viewer behind the images. */
 	UPROPERTY(EditAnywhere, Category = "Background")
-	FLinearColor ClearColor = FStyleColors::Panel.GetSpecifiedColor();
+	FLinearColor ClearColor = FStyleColors::Recessed.GetSpecifiedColor();
 
 	/** Texture tiled behind the images, on top of the clear color. */
 	UPROPERTY(EditAnywhere, Category = "Background")
