Media Viewer: Viewports
- Added viewports to the media viewer.

diff --git a/Engine/Plugins/Experimental/MediaViewer/Source/MediaViewer/MediaViewer.build.cs b/Engine/Plugins/Experimental/MediaViewer/Source/MediaViewer/MediaViewer.build.cs
--- a/Engine/Plugins/Experimental/MediaViewer/Source/MediaViewer/MediaViewer.build.cs
+++ b/Engine/Plugins/Experimental/MediaViewer/Source/MediaViewer/MediaViewer.build.cs
@@ -1,6 +1,5 @@
 // Copyright Epic Games, Inc. All Rights Reserved.
 
-using System.IO;
 using UnrealBuildTool;
 
 public class MediaViewer : ModuleRules
@@ -32,6 +31,7 @@ public class MediaViewer : ModuleRules
 				"Projects",
 				"PropertyEditor",
 				"RenderCore",
+				"RHI",
 				"SlateCore",
 				"Slate",
 				"ToolWidgets",
diff --git a/Engine/Plugins/Experimental/MediaViewer/Source/MediaViewer/Private/ImageViewers/LevelEditorViewportImageViewer.cpp b/Engine/Plugins/Experimental/MediaViewer/Source/MediaViewer/Private/ImageViewers/LevelEditorViewportImageViewer.cpp
new file mode 100644
--- /dev/null
+++ b/Engine/Plugins/Experimental/MediaViewer/Source/MediaViewer/Private/ImageViewers/LevelEditorViewportImageViewer.cpp
@@ -0,0 +1,246 @@
+// Copyright Epic Games, Inc. All Rights Reserved.
+
+#include "ImageViewers/LevelEditorViewportImageViewer.h"
+
+#include "Engine/TextureRenderTarget2D.h"
+#include "LevelEditor.h"
+#include "MediaViewerUtils.h"
+#include "Modules/ModuleManager.h"
+#include "Slate/SceneViewport.h"
+#include "SLevelViewport.h"
+
+#define LOCTEXT_NAMESPACE "LevelEditorViewportImageViewer"
+
+namespace UE::MediaViewer::Private
+{
+
+/** Can take either a string value (Viewport0..3) or a config key and return the string value (Viewport0..3) */
+FString ConfigKeyToStringValue(const FString& InConfigKey)
+{
+	int32 DotIndex;
+	InConfigKey.FindLastChar('.', DotIndex);
+
+	if (DotIndex < 0 || DotIndex >= (InConfigKey.Len() - 1))
+	{
+		return InConfigKey;
+	}
+
+	return InConfigKey.Mid(DotIndex + 1);
+}
+
+/** Could be a config key or a direct value (Viewport0..3) */
+FText GetViewportDisplayName(const FString& InViewportValue)
+{
+	return FText::FromString(ConfigKeyToStringValue(InViewportValue));
+}
+
+FText GetViewportTooltip(const FString& InViewportValue)
+{
+	return FText::GetEmpty();
+}
+
+const FLazyName FLevelEditorViewportImageViewer::ItemTypeName = TEXT("LevelEditorViewport");
+
+bool FLevelEditorViewportImageViewer::FFactory::SupportsItemType(FName InItemType) const
+{
+	return InItemType == ItemTypeName;
+}
+
+TSharedPtr<FMediaViewerLibraryItem> FLevelEditorViewportImageViewer::FFactory::CreateLibraryItem(const FMediaViewerLibraryItem& InSavedItem) const
+{
+	return MakeShared<FLevelEditorViewportImageViewer::FItem>(FPrivateToken(), InSavedItem);
+}
+
+FGuid FLevelEditorViewportImageViewer::FItem::GetIdForViewport(const FString& InConfigKey, bool bCreateIfInvalid)
+{
+	if (!InConfigKey.IsEmpty())
+	{
+		const TCHAR LastChar = InConfigKey[InConfigKey.Len() - 1];
+
+		switch (LastChar)
+		{
+			case '0':
+				return FGuid(0xFDD736C4, 0xFD7A98CB, 0x7C28F833, 0xD00B4BB3);
+
+			case '1':
+				return FGuid(0xB7929682, 0x749CED27, 0x357D74D5, 0xC5E19053);
+
+			case '2':
+				return FGuid(0xD82E7C42, 0xB038A8A1, 0x9C608700, 0x9C3B589C);
+
+			case '3':
+				return FGuid(0xF0FE7988, 0xE5DE498B, 0x4AC6026B, 0x255B3BEC);
+		}
+	}
+
+	if (bCreateIfInvalid)
+	{
+		return FGuid::NewGuid();
+	}
+
+	return FGuid();
+}
+
+TSharedPtr<FSceneViewport> FLevelEditorViewportImageViewer::FItem::GetViewportFromConfigKey(const FString& InConfigKey)
+{
+	return GetViewportFromStringValue(ConfigKeyToStringValue(InConfigKey));
+}
+
+TSharedPtr<FSceneViewport> FLevelEditorViewportImageViewer::FItem::GetViewportFromStringValue(const FString& InStringValue)
+{
+	if (InStringValue.IsEmpty())
+	{
+		return nullptr;
+	}
+
+	FLevelEditorModule& LevelEditorModule = FModuleManager::Get().GetModuleChecked<FLevelEditorModule>("LevelEditor");
+
+	TSharedPtr<ILevelEditor> LevelEditor = LevelEditorModule.GetFirstLevelEditor();
+
+	if (!LevelEditor.IsValid())
+	{
+		return nullptr;
+	}
+
+	TArray<TSharedPtr<SLevelViewport>> Viewports = LevelEditor->GetViewports();
+
+	for (const TSharedPtr<SLevelViewport>& Viewport : Viewports)
+	{
+		const FString ViewportValue = ConfigKeyToStringValue(Viewport->GetConfigKey().ToString());
+
+		if (ViewportValue == InStringValue)
+		{
+			return Viewport->GetSharedActiveViewport();
+		}
+	}
+
+	return nullptr;
+}
+
+FLevelEditorViewportImageViewer::FItem::FItem(const FString& InConfigKey)
+	: FItem(GetIdForViewport(InConfigKey, /* Create id if invalid */ true), InConfigKey)
+{
+}
+
+FLevelEditorViewportImageViewer::FItem::FItem(const FGuid& InId, const FString& InConfigKey)
+	: FMediaViewerLibraryItem(InId, GetViewportDisplayName(InConfigKey), GetViewportTooltip(InConfigKey), /* Transient */ false, ConfigKeyToStringValue(InConfigKey))
+{
+}
+
+FLevelEditorViewportImageViewer::FItem::FItem(FPrivateToken&& InPrivateToken, const FMediaViewerLibraryItem& InItem)
+	: FMediaViewerLibraryItem(InItem.GetId(), InItem.Name, InItem.ToolTip, /* Transient */ false, InItem.GetStringValue())
+{
+}
+
+FName FLevelEditorViewportImageViewer::FItem::GetItemType() const
+{
+	return FLevelEditorViewportImageViewer::ItemTypeName;
+}
+
+FText FLevelEditorViewportImageViewer::FItem::GetItemTypeDisplayName() const
+{
+	return LOCTEXT("LevelEditorViewport", "Level Editor Viewport");
+}
+
+TSharedPtr<FSlateBrush> FLevelEditorViewportImageViewer::FItem::CreateThumbnail()
+{
+	if (StringValue.IsEmpty())
+	{
+		return nullptr;
+	}
+
+	TSharedPtr<FSceneViewport> Viewport = GetViewportFromStringValue(StringValue);
+
+	if (!Viewport.IsValid())
+	{
+		return nullptr;
+	}
+
+	const FIntPoint RenderTargetSize = {32, 32};
+	RenderTarget = FMediaViewerUtils::CreateRenderTarget(RenderTargetSize, /* Transparent */ false);
+
+	RenderViewport(TNotNull<FViewport*>(Viewport.Get()), TNotNull<UTextureRenderTarget2D*>(RenderTarget), /* Resize render target */ false);
+
+	return MakeShared<FSlateImageBrush>(RenderTarget, FVector2D(RenderTarget->GetSurfaceWidth(), RenderTarget->GetSurfaceDepth()));
+}
+
+TSharedPtr<FMediaImageViewer> FLevelEditorViewportImageViewer::FItem::CreateImageViewer() const
+{
+	if (StringValue.IsEmpty())
+	{
+		return nullptr;
+	}
+
+	if (Id.IsValid())
+	{
+		return MakeShared<FLevelEditorViewportImageViewer>(Id, StringValue);
+	}
+
+	return MakeShared<FLevelEditorViewportImageViewer>(StringValue);
+}
+
+TSharedPtr<FMediaViewerLibraryItem> FLevelEditorViewportImageViewer::FItem::Clone() const
+{
+	if (StringValue.IsEmpty())
+	{
+		return nullptr;
+	}
+
+	return MakeShared<FLevelEditorViewportImageViewer::FItem>(FPrivateToken(), *this);
+}
+
+FString FLevelEditorViewportImageViewer::FItem::GetReferencerName() const
+{
+	static const FString ReferencerName = TEXT("FLevelEditorViewportImageViewer");
+	return ReferencerName;
+}
+
+void FLevelEditorViewportImageViewer::FItem::AddReferencedObjects(FReferenceCollector& Collector)
+{
+	if (RenderTarget)
+	{
+		Collector.AddReferencedObject(RenderTarget);
+	}
+}
+
+FLevelEditorViewportImageViewer::FLevelEditorViewportImageViewer(const FString& InViewportValue)
+	: FLevelEditorViewportImageViewer(FItem::GetIdForViewport(InViewportValue, /* Create id if invalid */ true), InViewportValue)
+{
+}
+
+FLevelEditorViewportImageViewer::FLevelEditorViewportImageViewer(const FGuid& InId, const FString& InViewportValue)
+	: FSceneViewportImageViewer(InId, FItem::GetViewportFromConfigKey(InViewportValue), GetViewportDisplayName(InViewportValue))
+	, StringValue(ConfigKeyToStringValue(InViewportValue))
+{
+}
+
+TSharedPtr<FMediaViewerLibraryItem> FLevelEditorViewportImageViewer::CreateLibraryItem() const
+{
+	return MakeShared<FItem>(ImageInfo.Id, StringValue);
+}
+
+void FLevelEditorViewportImageViewer::PaintImage(FMediaImagePaintParams& InPaintParams, const FMediaImagePaintGeometry& InPaintGeometry)
+{
+	if (ViweportSettings.bRealTime && !ViewportWeak.IsValid())
+	{
+		if (StringValue.IsEmpty())
+		{
+			return;
+		}
+
+		TSharedPtr<FSceneViewport> NewViewport = FItem::GetViewportFromStringValue(StringValue);
+
+		if (!NewViewport.IsValid())
+		{
+			return;
+		}
+
+		ViewportWeak = NewViewport;
+	}
+
+	FSceneViewportImageViewer::PaintImage(InPaintParams, InPaintGeometry);
+}
+
+} // UE::MediaViewer::Private
+
+#undef LOCTEXT_NAMESPACE
diff --git a/Engine/Plugins/Experimental/MediaViewer/Source/MediaViewer/Private/ImageViewers/LevelEditorViewportImageViewer.h b/Engine/Plugins/Experimental/MediaViewer/Source/MediaViewer/Private/ImageViewers/LevelEditorViewportImageViewer.h
new file mode 100644
--- /dev/null
+++ b/Engine/Plugins/Experimental/MediaViewer/Source/MediaViewer/Private/ImageViewers/LevelEditorViewportImageViewer.h
@@ -0,0 +1,83 @@
+// Copyright Epic Games, Inc. All Rights Reserved.
+
+#pragma once
+
+#include "ImageViewer/IMediaImageViewerFactory.h"
+#include "ImageViewers/SceneViewportImageViewer.h"
+#include "Library/MediaViewerLibraryItem.h"
+#include "UObject/GCObject.h"
+
+namespace UE::MediaViewer::Private
+{
+
+class FLevelEditorViewportImageViewer : public FSceneViewportImageViewer
+{
+public:
+	struct FFactory : public IMediaImageViewerFactory
+	{
+		FFactory()
+		{
+			Priority = 5000;
+		}
+
+		virtual bool SupportsAsset(const FAssetData& InAssetData) const override { return false; }
+		virtual TSharedPtr<FMediaImageViewer> CreateImageViewer(const FAssetData& InAssetData) const override { return nullptr; }
+		virtual TSharedPtr<FMediaViewerLibraryItem> CreateLibraryItem(const FAssetData& InAssetData) const override { return nullptr; }
+
+		virtual bool SupportsObject(TNotNull<UObject*> InObject) const override { return false; }
+		virtual TSharedPtr<FMediaImageViewer> CreateImageViewer(TNotNull<UObject*> InObject) const override { return nullptr; }
+		virtual TSharedPtr<FMediaViewerLibraryItem> CreateLibraryItem(TNotNull<UObject*> InObject) const override { return nullptr; }
+
+		virtual bool SupportsItemType(FName InItemType) const override;
+		virtual TSharedPtr<FMediaViewerLibraryItem> CreateLibraryItem(const FMediaViewerLibraryItem& InSavedItem) const override;
+	};
+
+	struct FItem : public FMediaViewerLibraryItem, public FGCObject
+	{
+		static FGuid GetIdForViewport(const FString& InConfigKey, bool bCreateIfInvalid);
+		static TSharedPtr<FSceneViewport> GetViewportFromConfigKey(const FString& InConfigKey);
+		static TSharedPtr<FSceneViewport> GetViewportFromStringValue(const FString& InStringValue);
+
+		TObjectPtr<UTextureRenderTarget2D> RenderTarget;
+
+		FItem(const FString& InConfigKey);
+
+		FItem(const FGuid& InId, const FString& InConfigKey);
+
+		FItem(FPrivateToken&& InPrivateToken, const FMediaViewerLibraryItem& InItem);
+
+		virtual ~FItem() override = default;
+
+		//~ Begin FMediaViewerLibraryItem
+		virtual FName GetItemType() const override;
+		virtual FText GetItemTypeDisplayName() const override;
+		virtual TSharedPtr<FSlateBrush> CreateThumbnail() override;
+		virtual TSharedPtr<FMediaImageViewer> CreateImageViewer() const override;
+		virtual TSharedPtr<FMediaViewerLibraryItem> Clone() const override;
+		//~ End FMediaViewerLibraryItem
+
+		//~ Begin FGCObject
+		virtual FString GetReferencerName() const override;
+		virtual void AddReferencedObjects(FReferenceCollector& Collector) override;
+		//~ End FGCObject
+	};
+
+	static const FLazyName ItemTypeName;
+
+	/** InStringValue should be Viewport0 to Viewport3 */
+	FLevelEditorViewportImageViewer(const FString& InStringValue);
+	FLevelEditorViewportImageViewer(const FGuid& InId, const FString& InStringValue);
+
+	//~ Begin FMediaImageViewer
+	virtual TSharedPtr<FMediaViewerLibraryItem> CreateLibraryItem() const override;
+	//~ End FMediaImageViewer
+	
+protected:
+	FString StringValue;
+
+	//~ Begin FMediaImageViewer
+	virtual void PaintImage(FMediaImagePaintParams& InPaintParams, const FMediaImagePaintGeometry& InPaintGeometry) override;
+	//~ End FMediaImageViewer
+};
+
+} // UE::MediaViewer::Private
diff --git a/Engine/Plugins/Experimental/MediaViewer/Source/MediaViewer/Private/ImageViewers/MediaImageViewer.cpp b/Engine/Plugins/Experimental/MediaViewer/Source/MediaViewer/Private/ImageViewers/MediaImageViewer.cpp
--- a/Engine/Plugins/Experimental/MediaViewer/Source/MediaViewer/Private/ImageViewers/MediaImageViewer.cpp
+++ b/Engine/Plugins/Experimental/MediaViewer/Source/MediaViewer/Private/ImageViewers/MediaImageViewer.cpp
@@ -3,6 +3,7 @@
 #include "ImageViewer/MediaImageViewer.h"
 
 #include "Engine/Texture.h"
+#include "Engine/TextureRenderTarget2D.h"
 #include "Math/Vector2D.h"
 #include "Misc/TVariant.h"
 #include "Widgets/SMediaViewer.h"
@@ -29,6 +30,11 @@ FMediaImageViewer::~FMediaImageViewer()
 	static_cast<FStructOnScope&>(PaintSettings.Get()).Destroy();
 }
 
+void FMediaImageViewer::UpdateId(const FGuid& InId)
+{
+	ImageInfo.Id = InId;
+}
+
 void FMediaImageViewer::NotifyPostChange(const FPropertyChangedEvent& InPropertyChangedEvent, FProperty* InPropertyThatChanged)
 {
 	if (UTexture* Texture = GetPanelSettings().BackgroundTexture.LoadSynchronous())
@@ -157,12 +163,24 @@ void FMediaImageViewer::PaintPanel(FMediaImagePaintParams& InPaintParams, const
 
 void FMediaImageViewer::PaintImage(FMediaImagePaintParams& InPaintParams, const FMediaImagePaintGeometry& InPaintGeometry)
 {
+	if (!Brush.IsValid())
+	{
+		return;
+	}
+
+	ESlateDrawEffect DrawEffects = ESlateDrawEffect::NoPixelSnapping;
+	
+	if (UTextureRenderTarget2D* RenderTarget = Cast<UTextureRenderTarget2D>(Brush->GetResourceObject()))
+	{
+		DrawEffects |= ESlateDrawEffect::PreMultipliedAlpha | ESlateDrawEffect::NoGamma;
+	}
+
 	FSlateDrawElement::MakeBox(
 		InPaintParams.DrawElements,
 		InPaintParams.LayerId,
 		InPaintGeometry.Geometry,
 		Brush.Get(),
-		ESlateDrawEffect::NoPixelSnapping,
+		DrawEffects,
 		GetPaintSettings().Tint * FLinearColor(1.f, 1.f, 1.f, InPaintParams.ImageOpacity)
 	);
 
diff --git a/Engine/Plugins/Experimental/MediaViewer/Source/MediaViewer/Private/ImageViewers/MediaTextureImageViewer.cpp b/Engine/Plugins/Experimental/MediaViewer/Source/MediaViewer/Private/ImageViewers/MediaTextureImageViewer.cpp
--- a/Engine/Plugins/Experimental/MediaViewer/Source/MediaViewer/Private/ImageViewers/MediaTextureImageViewer.cpp
+++ b/Engine/Plugins/Experimental/MediaViewer/Source/MediaViewer/Private/ImageViewers/MediaTextureImageViewer.cpp
@@ -169,6 +169,16 @@ FMediaTextureImageViewer::FMediaTextureImageViewer(const FGuid& InId, TNotNull<U
 	})
 {
 	MediaTextureSettings.MediaTexture = InMediaTexture;
+
+	const FIntPoint Size = {
+		FMath::RoundToInt(InMediaTexture->GetSurfaceWidth()),
+		FMath::RoundToInt(InMediaTexture->GetSurfaceHeight())
+	};
+
+	if (Size.X < 2 || Size.Y < 2)
+	{
+		return;
+	}
 	
 	MediaTextureSettings.TextureRenderMaterial = UMaterialInstanceDynamic::Create(
 		FMediaViewerUtils::GetTextureRenderMaterial(),
@@ -263,10 +273,21 @@ void FMediaTextureImageViewer::PaintImage(FMediaImagePaintParams& InPaintParams,
 {
 	if (MediaTextureSettings.bRealTime)
 	{
-		FMediaViewerUtils::RenderMaterial(
-			MediaTextureSettings.TextureRenderMaterial,
-			MediaTextureSettings.RenderTarget
-		);
+		if (UMediaTexture* MediaTexture = MediaTextureSettings.MediaTexture)
+		{
+			const FIntPoint Size = {
+				FMath::RoundToInt(MediaTexture->GetSurfaceWidth()),
+				FMath::RoundToInt(MediaTexture->GetSurfaceHeight())
+			};
+
+			if (Size.X > 1 && Size.Y > 1)
+			{
+				FMediaViewerUtils::RenderMaterial(
+					MediaTextureSettings.TextureRenderMaterial,
+					MediaTextureSettings.RenderTarget
+				);
+			}
+		}
 	}
 
 	FMediaImageViewer::PaintImage(InPaintParams, InPaintGeometry);
diff --git a/Engine/Plugins/Experimental/MediaViewer/Source/MediaViewer/Private/ImageViewers/SceneViewportImageViewer.cpp b/Engine/Plugins/Experimental/MediaViewer/Source/MediaViewer/Private/ImageViewers/SceneViewportImageViewer.cpp
new file mode 100644
--- /dev/null
+++ b/Engine/Plugins/Experimental/MediaViewer/Source/MediaViewer/Private/ImageViewers/SceneViewportImageViewer.cpp
@@ -0,0 +1,32 @@
+// Copyright Epic Games, Inc. All Rights Reserved.
+
+#include "ImageViewers/SceneViewportImageViewer.h"
+
+#include "Slate/SceneViewport.h"
+
+namespace UE::MediaViewer::Private
+{
+
+FSceneViewportImageViewer::FSceneViewportImageViewer(const TSharedPtr<FSceneViewport>& InViewport, const FText& InDisplayName)
+	: FSceneViewportImageViewer(FGuid::NewGuid(), InViewport, InDisplayName)
+{
+}
+
+FSceneViewportImageViewer::FSceneViewportImageViewer(const FGuid& InId, const TSharedPtr<FSceneViewport>& InViewport, const FText& InDisplayName)
+	: FViewportImageViewer({
+		InId,
+		InViewport.IsValid() ? InViewport->GetSize() : FIntPoint(0, 0),
+		/* Mip Count */ 1,
+		InDisplayName
+	})
+	, ViewportWeak(InViewport)
+{
+	CreateBrush();
+}
+
+FViewport* FSceneViewportImageViewer::GetViewport() const
+{
+	return ViewportWeak.Pin().Get();
+}
+
+} // UE::MediaViewer::Private
diff --git a/Engine/Plugins/Experimental/MediaViewer/Source/MediaViewer/Private/ImageViewers/SceneViewportImageViewer.h b/Engine/Plugins/Experimental/MediaViewer/Source/MediaViewer/Private/ImageViewers/SceneViewportImageViewer.h
new file mode 100644
--- /dev/null
+++ b/Engine/Plugins/Experimental/MediaViewer/Source/MediaViewer/Private/ImageViewers/SceneViewportImageViewer.h
@@ -0,0 +1,26 @@
+// Copyright Epic Games, Inc. All Rights Reserved.
+
+#pragma once
+
+#include "ImageViewers/ViewportImageViewer.h"
+
+class FSceneViewport;
+
+namespace UE::MediaViewer::Private
+{
+
+class FSceneViewportImageViewer : public FViewportImageViewer
+{
+public:
+	FSceneViewportImageViewer(const TSharedPtr<FSceneViewport>& InViewport, const FText& InName);
+	FSceneViewportImageViewer(const FGuid& InId, const TSharedPtr<FSceneViewport>& InViewport, const FText& InDisplayName);
+
+protected:
+	TWeakPtr<FSceneViewport> ViewportWeak;
+
+	//~ Begin FViewportImageViewer
+	virtual FViewport* GetViewport() const override;
+	//~ End FViewportImageViewer
+};
+
+} // UE::MediaViewer::Private
diff --git a/Engine/Plugins/Experimental/MediaViewer/Source/MediaViewer/Private/ImageViewers/ViewportImageViewer.cpp b/Engine/Plugins/Experimental/MediaViewer/Source/MediaViewer/Private/ImageViewers/ViewportImageViewer.cpp
new file mode 100644
--- /dev/null
+++ b/Engine/Plugins/Experimental/MediaViewer/Source/MediaViewer/Private/ImageViewers/ViewportImageViewer.cpp
@@ -0,0 +1,186 @@
+// Copyright Epic Games, Inc. All Rights Reserved.
+
+#include "ImageViewers/ViewportImageViewer.h"
+
+#include "Brushes/SlateImageBrush.h"
+#include "Engine/TextureRenderTarget2D.h"
+#include "GlobalShader.h"
+#include "Math/Vector2D.h"
+#include "MediaViewerUtils.h"
+#include "Misc/Optional.h"
+#include "Misc/TVariant.h"
+#include "RenderGraphBuilder.h"
+#include "RenderGraphUtils.h"
+#include "TextureResource.h"
+#include "UnrealClient.h"
+
+namespace UE::MediaViewer::Private
+{
+
+FViewportImageViewer::FViewportImageViewer(const FMediaImageViewerInfo& InImageInfo)
+	: FMediaImageViewer(InImageInfo)
+{
+}
+
+FViewportImageViewer::~FViewportImageViewer()
+{
+}
+
+bool FViewportImageViewer::IsValid() const
+{
+	return FMediaImageViewer::IsValid() && !!GetViewport();
+}
+
+TOptional<TVariant<FColor, FLinearColor>> FViewportImageViewer::GetPixelColor(const FIntPoint& InPixelCoords, int32 InMipLevel) const
+{
+	if (InPixelCoords.X < 0 || InPixelCoords.Y < 0)
+	{
+		return {};
+	}
+
+	UTextureRenderTarget2D* RenderTarget = ViweportSettings.RenderTarget;
+
+	if (!RenderTarget)
+	{
+		return {};
+	}
+
+	const FIntPoint RenderTargetSize = {static_cast<int32>(RenderTarget->Source.GetSizeX()), static_cast<int32>(RenderTarget->Source.GetSizeY())};
+
+	if (InPixelCoords.X >= RenderTargetSize.X || InPixelCoords.Y >= RenderTargetSize.Y)
+	{
+		return {};
+	}
+
+	const FImage* CachedMip = MipCache.GetMipImage(InMipLevel);
+
+	if (!CachedMip)
+	{
+		return {};
+	}
+
+	TVariant<FColor, FLinearColor> Pixel;
+	Pixel.Set<FLinearColor>(CachedMip->GetOnePixelLinear(InPixelCoords.X, InPixelCoords.Y));
+
+	return Pixel;
+}
+
+TSharedPtr<FStructOnScope> FViewportImageViewer::GetCustomSettingsOnScope() const
+{
+	return MakeShared<FStructOnScope>(
+		FViewportImageViewerSettings::StaticStruct(), 
+		reinterpret_cast<uint8*>(&const_cast<FViewportImageViewerSettings&>(ViweportSettings))
+	);
+}
+
+void FViewportImageViewer::PaintImage(FMediaImagePaintParams& InPaintParams, const FMediaImagePaintGeometry& InPaintGeometry)
+{
+	if (ViweportSettings.bRealTime)
+	{
+		RenderViewport();
+	}
+
+	FMediaImageViewer::PaintImage(InPaintParams, InPaintGeometry);
+}
+
+FString FViewportImageViewer::GetReferencerName() const
+{
+	static const FString ReferencerName = TEXT("FViewportImageViewer");
+	return ReferencerName;
+}
+
+void FViewportImageViewer::AddReferencedObjects(FReferenceCollector& Collector)
+{
+	Collector.AddPropertyReferencesWithStructARO(
+		FViewportImageViewerSettings::StaticStruct(),
+		&ViweportSettings
+	);
+}
+
+void FViewportImageViewer::CreateBrush()
+{
+	FViewport* Viewport = GetViewport();
+
+	if (!Viewport)
+	{
+		return;
+	}
+
+	ImageInfo.Size = Viewport->GetSizeXY();
+
+	UTextureRenderTarget2D* RenderTarget = FMediaViewerUtils::CreateRenderTarget(ImageInfo.Size, /* Transparent */ false);
+	ViweportSettings.RenderTarget = RenderTarget;
+
+	MipCache = FTextureMipCache(RenderTarget);
+
+	Brush = MakeShared<FSlateImageBrush>(RenderTarget, FVector2D(RenderTarget->GetSurfaceWidth(), RenderTarget->GetSurfaceDepth()));
+}
+
+void FViewportImageViewer::RenderViewport(TNotNull<FViewport*> InViewport, TNotNull<UTextureRenderTarget2D*> InRenderTarget, 
+	bool bResizeTargetToViewport)
+{
+	if (bResizeTargetToViewport)
+	{
+		const FIntPoint ViewportSize = InViewport->GetSizeXY();
+
+		if (ViewportSize.X != InRenderTarget->GetSurfaceWidth()
+			|| ViewportSize.Y != InRenderTarget->GetSurfaceHeight())
+		{
+			InRenderTarget->ResizeTarget(ViewportSize.X, ViewportSize.Y);
+		}
+	}
+
+	ENQUEUE_RENDER_COMMAND(CopyViewportRenderTarget)(
+		[Viewport = InViewport, RenderTarget = InRenderTarget](FRHICommandListImmediate& InRHICmdList)
+		{
+			FTextureRenderTargetResource* RenderTargetResource = RenderTarget->GetRenderTargetResource();;
+
+			if (!RenderTargetResource)
+			{
+				return;
+			}
+
+			FRDGBuilder GraphBuilder(InRHICmdList);
+
+			FRDGTexture* InputTexture = Viewport->GetRenderTargetTexture(GraphBuilder);
+
+			if (!InputTexture)
+			{
+				return;
+			}
+
+			FRDGTexture* OutputTexture = RenderTargetResource->GetRenderTargetTexture(GraphBuilder);
+
+			if (!OutputTexture)
+			{
+				return;
+			}
+
+			const FGlobalShaderMap* GlobalShaderMap = GetGlobalShaderMap(GMaxRHIFeatureLevel);
+
+			AddDrawTexturePass(GraphBuilder, GlobalShaderMap, InputTexture, OutputTexture, FRDGDrawTextureInfo());
+
+			GraphBuilder.Execute();
+		}
+	);
+}
+
+void FViewportImageViewer::RenderViewport()
+{
+	FViewport* Viewport = GetViewport();
+
+	if (!Viewport || !::IsValid(ViweportSettings.RenderTarget))
+	{
+		return;
+	}
+
+	ImageInfo.Size = Viewport->GetSizeXY();
+
+	RenderViewport(TNotNull<FViewport*>(Viewport), TNotNull<UTextureRenderTarget2D*>(ViweportSettings.RenderTarget));
+
+	MipCache.Invalidate();
+}
+
+} // UE::MediaViewer::Private
+
+#undef LOCTEXT_NAMESPACE
diff --git a/Engine/Plugins/Experimental/MediaViewer/Source/MediaViewer/Private/ImageViewers/ViewportImageViewer.h b/Engine/Plugins/Experimental/MediaViewer/Source/MediaViewer/Private/ImageViewers/ViewportImageViewer.h
new file mode 100644
--- /dev/null
+++ b/Engine/Plugins/Experimental/MediaViewer/Source/MediaViewer/Private/ImageViewers/ViewportImageViewer.h
@@ -0,0 +1,70 @@
+// Copyright Epic Games, Inc. All Rights Reserved.
+
+#pragma once
+
+#include "ImageViewer/IMediaImageViewerFactory.h"
+#include "ImageViewer/MediaImageViewer.h"
+#include "UObject/GCObject.h"
+
+#include "ImageViewers/TextureMipCache.h"
+#include "IMediaViewerModule.h"
+#include "Templates/SharedPointer.h"
+
+#include "ViewportImageViewer.generated.h"
+
+class FTextureRenderTargetResource;
+class FViewport;
+class UTextureRenderTarget2D;
+
+USTRUCT()
+struct FViewportImageViewerSettings 
+{
+	GENERATED_BODY()
+
+	UPROPERTY(VisibleAnywhere, Category = "Material")
+	TObjectPtr<UTextureRenderTarget2D> RenderTarget;
+
+	/** Will render the material every frame. */
+	UPROPERTY(EditAnywhere, Category = "Material")
+	bool bRealTime;
+};
+
+namespace UE::MediaViewer::Private
+{
+
+class FViewportImageViewer : public FMediaImageViewer, public FGCObject
+{
+public:
+	FViewportImageViewer(const FMediaImageViewerInfo& InImageInfo);
+
+	virtual ~FViewportImageViewer() override;
+
+	//~ Begin FMediaImageViewer
+	virtual bool IsValid() const override;
+	virtual TOptional<TVariant<FColor, FLinearColor>> GetPixelColor(const FIntPoint& InPixelCoords, int32 InMipLevel) const override;
+	virtual TSharedPtr<FStructOnScope> GetCustomSettingsOnScope() const override;
+	//~ End FMediaImageViewer
+
+	//~ Begin FGCObject
+	virtual FString GetReferencerName() const override;
+	virtual void AddReferencedObjects(FReferenceCollector& Collector) override;
+	//~ End FGCObject
+
+protected:
+	static void RenderViewport(TNotNull<FViewport*> InViewport, TNotNull<UTextureRenderTarget2D*> InRenderTarget, bool bResizeTargetToViewport = true);
+
+	FViewportImageViewerSettings ViweportSettings;
+	FTextureMipCache MipCache;
+
+	virtual FViewport* GetViewport() const = 0;
+
+	void CreateBrush();
+
+	void RenderViewport();
+
+	//~ Begin FMediaImageViewer
+	virtual void PaintImage(FMediaImagePaintParams& InPaintParams, const FMediaImagePaintGeometry& InPaintGeometry) override;
+	//~ End FMediaImageViewer
+};
+
+} // UE::MediaViewer::Private
diff --git a/Engine/Plugins/Experimental/MediaViewer/Source/MediaViewer/Private/Library/LevelEditorViewportGroup.cpp b/Engine/Plugins/Experimental/MediaViewer/Source/MediaViewer/Private/Library/LevelEditorViewportGroup.cpp
new file mode 100644
--- /dev/null
+++ b/Engine/Plugins/Experimental/MediaViewer/Source/MediaViewer/Private/Library/LevelEditorViewportGroup.cpp
@@ -0,0 +1,82 @@
+// Copyright Epic Games, Inc. All Rights Reserved.
+
+#include "Library/LevelEditorViewportGroup.h"
+
+#include "ImageViewers/LevelEditorViewportImageViewer.h"
+#include "LevelEditor.h"
+#include "Library/MediaViewerLibraryItem.h"
+#include "Modules/ModuleManager.h"
+#include "Slate/SceneViewport.h"
+#include "SLevelViewport.h"
+
+#define LOCTEXT_NAMESPACE "LevelEditorViewportGroup"
+
+namespace UE::MediaViewer::Private
+{
+
+FLevelEditorViewportGroup::FLevelEditorViewportGroup(const TSharedRef<IMediaViewerLibrary>& InLibrary)
+	: FLevelEditorViewportGroup(InLibrary, FGuid::NewGuid())
+{
+}
+
+FLevelEditorViewportGroup::FLevelEditorViewportGroup(const TSharedRef<IMediaViewerLibrary>& InLibrary,
+	const FGuid& InGuid)
+	: FMediaViewerLibraryDynamicGroup(
+		InLibrary,
+		InGuid,
+		LOCTEXT("LevelEditorViewports", "Editor Viewports"),
+		LOCTEXT("LevelEditorViewportsTooltip", "The viewports available in the Level Editor."),
+		FGenerateItems::CreateStatic(&FLevelEditorViewportGroup::GetLevelEditorViewportItems)
+	)
+{
+}
+
+TArray<TSharedRef<FMediaViewerLibraryItem>> FLevelEditorViewportGroup::GetLevelEditorViewportItems()
+{
+	TArray<TSharedRef<FMediaViewerLibraryItem>> LevelEditorViewportItems;
+
+	FLevelEditorModule& LevelEditorModule = FModuleManager::Get().GetModuleChecked<FLevelEditorModule>("LevelEditor");
+
+	TSharedPtr<ILevelEditor> LevelEditor = LevelEditorModule.GetFirstLevelEditor();
+
+	if (!LevelEditor.IsValid())
+	{
+		return LevelEditorViewportItems;
+	}
+
+	TArray<TSharedPtr<SLevelViewport>> Viewports = LevelEditor->GetViewports();
+
+	for (int32 Index = 0; Index < Viewports.Num(); ++Index)
+	{
+		TSharedPtr<FSceneViewport> ActiveViewport = Viewports[Index]->GetSharedActiveViewport();
+
+		if (!ActiveViewport.IsValid())
+		{
+			continue;
+		}
+
+		const FIntPoint Size = ActiveViewport->GetSize();
+
+		if (Size.X < 2 || Size.Y < 2)
+		{
+			continue;
+		}
+
+		const FString ConfigKey = Viewports[Index]->GetConfigKey().ToString();
+		const FGuid ViewportId = FLevelEditorViewportImageViewer::FItem::GetIdForViewport(
+			ConfigKey,
+			/* Create id if invalid */ false
+		);
+
+		if (ViewportId.IsValid())
+		{
+			LevelEditorViewportItems.Add(MakeShared<FLevelEditorViewportImageViewer::FItem>(ViewportId, ConfigKey));
+		}
+	}
+
+	return LevelEditorViewportItems;
+}
+
+} // UE::MediaViewer::Private
+
+#undef LOCTEXT_NAMESPACE
diff --git a/Engine/Plugins/Experimental/MediaViewer/Source/MediaViewer/Private/Library/LevelEditorViewportGroup.h b/Engine/Plugins/Experimental/MediaViewer/Source/MediaViewer/Private/Library/LevelEditorViewportGroup.h
new file mode 100644
--- /dev/null
+++ b/Engine/Plugins/Experimental/MediaViewer/Source/MediaViewer/Private/Library/LevelEditorViewportGroup.h
@@ -0,0 +1,24 @@
+// Copyright Epic Games, Inc. All Rights Reserved.
+
+#pragma once
+
+#include "Library/MediaViewerLibraryDynamicGroup.h"
+
+namespace UE::MediaViewer::Private
+{
+
+/**
+ * A group that generates entries based on available level editor viewports.
+ */
+struct FLevelEditorViewportGroup : public FMediaViewerLibraryDynamicGroup
+{
+	FLevelEditorViewportGroup(const TSharedRef<IMediaViewerLibrary>& InLibrary);
+	FLevelEditorViewportGroup(const TSharedRef<IMediaViewerLibrary>& InLibrary, const FGuid& InId);
+
+	virtual ~FLevelEditorViewportGroup() override = default;
+
+protected:
+	static TArray<TSharedRef<FMediaViewerLibraryItem>> GetLevelEditorViewportItems();
+};
+
+} // UE::MediaViewer::Private
diff --git a/Engine/Plugins/Experimental/MediaViewer/Source/MediaViewer/Private/Library/MediaTextureGroup.cpp b/Engine/Plugins/Experimental/MediaViewer/Source/MediaViewer/Private/Library/MediaTextureGroup.cpp
new file mode 100644
--- /dev/null
+++ b/Engine/Plugins/Experimental/MediaViewer/Source/MediaViewer/Private/Library/MediaTextureGroup.cpp
@@ -0,0 +1,56 @@
+// Copyright Epic Games, Inc. All Rights Reserved.
+
+#include "Library/MediaTextureGroup.h"
+
+#include "ImageViewers/MediaTextureImageViewer.h"
+#include "Library/MediaViewerLibraryItem.h"
+#include "MediaTexture.h"
+#include "UObject/UObjectIterator.h"
+
+#define LOCTEXT_NAMESPACE "MediaTextureGroup"
+
+namespace UE::MediaViewer::Private
+{
+
+FMediaTextureGroup::FMediaTextureGroup(const TSharedRef<IMediaViewerLibrary>& InLibrary)
+	: FMediaTextureGroup(InLibrary, FGuid::NewGuid())
+{
+}
+
+FMediaTextureGroup::FMediaTextureGroup(const TSharedRef<IMediaViewerLibrary>& InLibrary,
+	const FGuid& InGuid)
+	: FMediaViewerLibraryDynamicGroup(
+		InLibrary,
+		InGuid,
+		LOCTEXT("MediaTextures", "Media Textures"),
+		LOCTEXT("MediaTexturesTooltip", "The available media textures."),
+		FGenerateItems::CreateStatic(&FMediaTextureGroup::GetMediaTextureItems)
+	)
+{
+}
+
+TArray<TSharedRef<FMediaViewerLibraryItem>> FMediaTextureGroup::GetMediaTextureItems()
+{
+	TArray<TSharedRef<FMediaViewerLibraryItem>> MediaTextureItems;
+
+	TSharedRef<FMediaTextureImageViewer::FFactory> Factory = MakeShared<FMediaTextureImageViewer::FFactory>();
+
+	for (UMediaTexture* MediaTexture : TObjectRange<UMediaTexture>())
+	{
+		if (MediaTexture->IsTemplate())
+		{
+			continue;
+		}
+
+		if (TSharedPtr<FMediaViewerLibraryItem> NewItem = Factory->CreateLibraryItem(TNotNull<UMediaTexture*>(MediaTexture)))
+		{
+			MediaTextureItems.Add(NewItem.ToSharedRef());
+		}
+	}
+
+	return MediaTextureItems;
+}
+
+} // UE::MediaViewer::Private
+
+#undef LOCTEXT_NAMESPACE
diff --git a/Engine/Plugins/Experimental/MediaViewer/Source/MediaViewer/Private/Library/MediaTextureGroup.h b/Engine/Plugins/Experimental/MediaViewer/Source/MediaViewer/Private/Library/MediaTextureGroup.h
new file mode 100644
--- /dev/null
+++ b/Engine/Plugins/Experimental/MediaViewer/Source/MediaViewer/Private/Library/MediaTextureGroup.h
@@ -0,0 +1,24 @@
+// Copyright Epic Games, Inc. All Rights Reserved.
+
+#pragma once
+
+#include "Library/MediaViewerLibraryDynamicGroup.h"
+
+namespace UE::MediaViewer::Private
+{
+
+/**
+ * A group that generates entries based on available UMediaTextures.
+*/
+struct FMediaTextureGroup : public FMediaViewerLibraryDynamicGroup
+{
+	FMediaTextureGroup(const TSharedRef<IMediaViewerLibrary>& InLibrary);
+	FMediaTextureGroup(const TSharedRef<IMediaViewerLibrary>& InLibrary, const FGuid& InId);
+
+	virtual ~FMediaTextureGroup() override = default;
+
+protected:
+	static TArray<TSharedRef<FMediaViewerLibraryItem>> GetMediaTextureItems();
+};
+
+} // UE::MediaViewer::Private
diff --git a/Engine/Plugins/Experimental/MediaViewer/Source/MediaViewer/Private/Library/MediaViewerLibraryMediaTextureGroup.cpp b/Engine/Plugins/Experimental/MediaViewer/Source/MediaViewer/Private/Library/MediaViewerLibraryMediaTextureGroup.cpp
deleted file mode 100644
diff --git a/Engine/Plugins/Experimental/MediaViewer/Source/MediaViewer/Private/Library/MediaViewerLibraryMediaTextureGroup.h b/Engine/Plugins/Experimental/MediaViewer/Source/MediaViewer/Private/Library/MediaViewerLibraryMediaTextureGroup.h
deleted file mode 100644
diff --git a/Engine/Plugins/Experimental/MediaViewer/Source/MediaViewer/Private/Library/MediaViewerLibraryTextureRenderTarget2DGroup.cpp b/Engine/Plugins/Experimental/MediaViewer/Source/MediaViewer/Private/Library/MediaViewerLibraryTextureRenderTarget2DGroup.cpp
deleted file mode 100644
diff --git a/Engine/Plugins/Experimental/MediaViewer/Source/MediaViewer/Private/Library/MediaViewerLibraryTextureRenderTarget2DGroup.h b/Engine/Plugins/Experimental/MediaViewer/Source/MediaViewer/Private/Library/MediaViewerLibraryTextureRenderTarget2DGroup.h
deleted file mode 100644
diff --git a/Engine/Plugins/Experimental/MediaViewer/Source/MediaViewer/Private/Library/TextureRenderTarget2DGroup.cpp b/Engine/Plugins/Experimental/MediaViewer/Source/MediaViewer/Private/Library/TextureRenderTarget2DGroup.cpp
new file mode 100644
--- /dev/null
+++ b/Engine/Plugins/Experimental/MediaViewer/Source/MediaViewer/Private/Library/TextureRenderTarget2DGroup.cpp
@@ -0,0 +1,62 @@
+// Copyright Epic Games, Inc. All Rights Reserved.
+
+#include "Library/TextureRenderTarget2DGroup.h"
+
+#include "Engine/TextureRenderTarget2D.h"
+#include "ImageViewers/TextureRenderTarget2DImageViewer.h"
+#include "Library/MediaViewerLibraryItem.h"
+#include "MediaViewerUtils.h"
+#include "UObject/UObjectIterator.h"
+
+#define LOCTEXT_NAMESPACE "TextureRenderTarget2DGroup"
+
+namespace UE::MediaViewer::Private
+{
+
+FTextureRenderTarget2DGroup::FTextureRenderTarget2DGroup(const TSharedRef<IMediaViewerLibrary>& InLibrary)
+	: FTextureRenderTarget2DGroup(InLibrary, FGuid::NewGuid())
+{
+}
+
+FTextureRenderTarget2DGroup::FTextureRenderTarget2DGroup(const TSharedRef<IMediaViewerLibrary>& InLibrary, 
+	const FGuid& InGuid)
+	: FMediaViewerLibraryDynamicGroup(
+		InLibrary,
+		InGuid,
+		LOCTEXT("RenderTargets", "Render Targets"),
+		LOCTEXT("RenderTargetsTooltip", "The available render targets."),
+		FGenerateItems::CreateStatic(&FTextureRenderTarget2DGroup::GetTextureRenderTarget2DItems)
+	)
+{
+}
+
+TArray<TSharedRef<FMediaViewerLibraryItem>> FTextureRenderTarget2DGroup::GetTextureRenderTarget2DItems()
+{
+	TArray<TSharedRef<FMediaViewerLibraryItem>> TextureRenderTarget2DItems;
+
+	TSharedRef<FTextureRenderTarget2DImageViewer::FFactory> Factory = MakeShared<FTextureRenderTarget2DImageViewer::FFactory>();
+
+	for (UTextureRenderTarget2D* TextureRenderTarget2D : TObjectRange<UTextureRenderTarget2D>())
+	{
+		if (TextureRenderTarget2D->IsTemplate())
+		{
+			continue;
+		}
+
+		if (TextureRenderTarget2D->HasAssetUserDataOfClass(UMediaViewerUserData::StaticClass()))
+		{
+			continue;
+		}
+
+		if (TSharedPtr<FMediaViewerLibraryItem> NewItem = Factory->CreateLibraryItem(TNotNull<UTextureRenderTarget2D*>(TextureRenderTarget2D)))
+		{
+			TextureRenderTarget2DItems.Add(NewItem.ToSharedRef());
+		}
+	}
+
+	return TextureRenderTarget2DItems;
+}
+
+} // UE::MediaViewer::Private
+
+#undef LOCTEXT_NAMESPACE
diff --git a/Engine/Plugins/Experimental/MediaViewer/Source/MediaViewer/Private/Library/TextureRenderTarget2DGroup.h b/Engine/Plugins/Experimental/MediaViewer/Source/MediaViewer/Private/Library/TextureRenderTarget2DGroup.h
new file mode 100644
--- /dev/null
+++ b/Engine/Plugins/Experimental/MediaViewer/Source/MediaViewer/Private/Library/TextureRenderTarget2DGroup.h
@@ -0,0 +1,24 @@
+// Copyright Epic Games, Inc. All Rights Reserved.
+
+#pragma once
+
+#include "Library/MediaViewerLibraryDynamicGroup.h"
+
+namespace UE::MediaViewer::Private
+{
+
+/**
+ * A group that generates entries based on available UTextureRenderTarget2Ds.
+*/
+struct FTextureRenderTarget2DGroup : public FMediaViewerLibraryDynamicGroup
+{
+	FTextureRenderTarget2DGroup(const TSharedRef<IMediaViewerLibrary>& InLibrary);
+	FTextureRenderTarget2DGroup(const TSharedRef<IMediaViewerLibrary>& InLibrary, const FGuid& InId);
+
+	virtual ~FTextureRenderTarget2DGroup() override = default;
+
+protected:
+	static TArray<TSharedRef<FMediaViewerLibraryItem>> GetTextureRenderTarget2DItems();
+};
+
+} // UE::MediaViewer::Private
diff --git a/Engine/Plugins/Experimental/MediaViewer/Source/MediaViewer/Private/MediaViewerUtils.cpp b/Engine/Plugins/Experimental/MediaViewer/Source/MediaViewer/Private/MediaViewerUtils.cpp
--- a/Engine/Plugins/Experimental/MediaViewer/Source/MediaViewer/Private/MediaViewerUtils.cpp
+++ b/Engine/Plugins/Experimental/MediaViewer/Source/MediaViewer/Private/MediaViewerUtils.cpp
@@ -11,6 +11,7 @@
 #include "IStructureDetailsView.h"
 #include "Materials/Material.h"
 #include "Materials/MaterialInterface.h"
+#include "Math/IntPoint.h"
 #include "Misc/TVariant.h"
 #include "Modules/ModuleManager.h"
 #include "PixelFormat.h"
@@ -153,18 +154,25 @@ TOptional<TVariant<FColor, FLinearColor>> FMediaViewerUtils::GetPixelColor(TArra
 	return ColorVariant;
 }
 
-UTextureRenderTarget2D* FMediaViewerUtils::RenderMaterial(TNotNull<UMaterialInterface*> InMaterial)
+UTextureRenderTarget2D* FMediaViewerUtils::CreateRenderTarget(const FIntPoint& InSize, bool bInTransparent)
 {
 	UTextureRenderTarget2D* RenderTarget = NewObject<UTextureRenderTarget2D>(GetTransientPackage());
 	check(RenderTarget);
 	RenderTarget->RenderTargetFormat = RTF_RGBA8;
-	RenderTarget->ClearColor = InMaterial->IsUIMaterial() ? FLinearColor::Transparent : FLinearColor::Black;
-	RenderTarget->bAutoGenerateMips = true;
+	RenderTarget->bAutoGenerateMips = false;
 	RenderTarget->bCanCreateUAV = false;
-	RenderTarget->InitAutoFormat(256, 256);
+	RenderTarget->ClearColor = bInTransparent ? FLinearColor::Transparent : FLinearColor::Black;
+	RenderTarget->InitAutoFormat(InSize.X, InSize.Y);
 	RenderTarget->UpdateResourceImmediate(true);
 	RenderTarget->AddAssetUserData(NewObject<UMediaViewerUserData>(RenderTarget));
 
+	return RenderTarget;
+}
+
+UTextureRenderTarget2D* FMediaViewerUtils::RenderMaterial(TNotNull<UMaterialInterface*> InMaterial)
+{
+	UTextureRenderTarget2D* RenderTarget = CreateRenderTarget(FIntPoint(256, 256), InMaterial->IsUIMaterial());
+
 	RenderMaterial(InMaterial, TNotNull<UTextureRenderTarget2D*>(RenderTarget));
 
 	return RenderTarget;
diff --git a/Engine/Plugins/Experimental/MediaViewer/Source/MediaViewer/Private/MediaViewerUtils.h b/Engine/Plugins/Experimental/MediaViewer/Source/MediaViewer/Private/MediaViewerUtils.h
--- a/Engine/Plugins/Experimental/MediaViewer/Source/MediaViewer/Private/MediaViewerUtils.h
+++ b/Engine/Plugins/Experimental/MediaViewer/Source/MediaViewer/Private/MediaViewerUtils.h
@@ -42,15 +42,29 @@ class FMediaViewerUtils
 public:
 	static constexpr const TCHAR* TextureRenderMaterial_TextureParameterName = TEXT("Texture");
 
+	/** Returns a UI Material asset that can be used to render a texture to the emissive and opacity channels. */
 	static UMaterialInterface* GetTextureRenderMaterial();
 
+	/** Retrieves the pixel color of a given pixel based on its pixel format. Or none if out of range. */
 	static TOptional<TVariant<FColor, FLinearColor>> GetPixelColor(TArrayView64<uint8> InPixelData, EPixelFormat InPixelFormat, 
 		const FIntPoint& InTextureSize, const FIntPoint& InPixelCoords, int32 InMipLevel);
 
+	/** 
+	 * Creates a render target texture parented to the transient package. 
+	 * The target will have the @see UMediaViewerUserData asset data to differentiate it from other render targets.
+	 */
+	static UTextureRenderTarget2D* CreateRenderTarget(const FIntPoint& InSize, bool bInTransparent);
+
+	/** 
+	 * Creates a render target (@see CreateRenderTarget) and renders the material to it.
+	 * By default it creates a 256x256 render target which is transparent if the material is a UI material.
+	 */
 	static UTextureRenderTarget2D* RenderMaterial(TNotNull<UMaterialInterface*> InMaterial);
 
+	/** Renders the given material to the given render target using FWidgetRenderer and an FSlateMaterialBrush. */
 	static void RenderMaterial(TNotNull<UMaterialInterface*> InMaterial, TNotNull<UTextureRenderTarget2D*> InRenderTarget);
 
+	/** Creates a struct details view based on the given struct. The view will have most settings disabled for a clean view. */
 	static TSharedRef<IStructureDetailsView> CreateStructDetailsView(TSharedRef<FStructOnScope> InStructOnScope, const FText& InCustomName,
 		FNotifyHook* InNotifyHook = nullptr);
 };
diff --git a/Engine/Plugins/Experimental/MediaViewer/Source/MediaViewer/Private/Widgets/SMediaViewer.cpp b/Engine/Plugins/Experimental/MediaViewer/Source/MediaViewer/Private/Widgets/SMediaViewer.cpp
--- a/Engine/Plugins/Experimental/MediaViewer/Source/MediaViewer/Private/Widgets/SMediaViewer.cpp
+++ b/Engine/Plugins/Experimental/MediaViewer/Source/MediaViewer/Private/Widgets/SMediaViewer.cpp
@@ -102,7 +102,7 @@ void SMediaViewer::Construct(const FArguments& InArgs, const FMediaViewerArgs& I
 
 	TSharedRef<SSidebar> Sidebar = SNew(SSidebar, SidebarContainer)
 		.TabLocation(ESidebarTabLocation::Left)
-		.InitialDrawerSize(0.5f)
+		.InitialDrawerSize(0.25f)
 		.OnGetContent(FOnGetContent::CreateSPLambda(
 			this,
 			[this]
@@ -112,7 +112,7 @@ void SMediaViewer::Construct(const FArguments& InArgs, const FMediaViewerArgs& I
 		);
 
 	FSidebarState SidebarState;
-	SidebarState.SetDrawerSizes(0.5f, 0.5f);
+	SidebarState.SetDrawerSizes(0.25f, 0.25f);
 
 	SidebarContainer->RebuildSidebar(Sidebar, SidebarState);
 
diff --git a/Engine/Plugins/Experimental/MediaViewer/Source/MediaViewer/Private/Widgets/SMediaViewerLibraryItem.cpp b/Engine/Plugins/Experimental/MediaViewer/Source/MediaViewer/Private/Widgets/SMediaViewerLibraryItem.cpp
--- a/Engine/Plugins/Experimental/MediaViewer/Source/MediaViewer/Private/Widgets/SMediaViewerLibraryItem.cpp
+++ b/Engine/Plugins/Experimental/MediaViewer/Source/MediaViewer/Private/Widgets/SMediaViewerLibraryItem.cpp
@@ -73,7 +73,7 @@ void SMediaViewerLibraryItem::Construct(const FArguments& InArgs, const TSharedR
 	TSharedRef<SHorizontalBox> Box = SNew(SHorizontalBox)
 		+ SHorizontalBox::Slot()
 		.AutoWidth()
-		.Padding(15.f, 5.f, 0.f, 5.f)
+		.Padding(0.f, 5.f, 0.f, 5.f)
 		.VAlign(VAlign_Center)
 		[
 			SNew(SImage)
@@ -127,7 +127,7 @@ void SMediaViewerLibraryItem::Construct(const FArguments& InArgs, const TSharedR
 
 	Box->AddSlot()
 		.AutoWidth()
-		.Padding(3.f, 5.f, 0.f, 5.f)
+		.Padding(3.f, 5.f, 3.f, 5.f)
 		.VAlign(VAlign_Center)
 		[
 			SNew(SButton)
@@ -147,7 +147,7 @@ void SMediaViewerLibraryItem::Construct(const FArguments& InArgs, const TSharedR
 		Box->AddSlot()
 			.AutoWidth()
 			.VAlign(VAlign_Center)
-			.Padding(3.f, 5.f, 5.f, 5.f)
+			.Padding(0.f, 5.f, 3.f, 5.f)
 			[
 				SNew(SButton)
 				.ContentPadding(FMargin(3.f, 5.f, 3.f, 6.f))
@@ -170,6 +170,7 @@ void SMediaViewerLibraryItem::Construct(const FArguments& InArgs, const TSharedR
 			.Visibility(this, &SMediaViewerLibraryItem::GetHoverVisibility)
 		]
 		+ SOverlay::Slot()
+		.Padding(0.f, 0.f, 2.f, 0.f)
 		[
 			Box
 		];
diff --git a/Engine/Plugins/Experimental/MediaViewer/Source/MediaViewer/Private/Widgets/SMediaViewerLibraryPrivate.cpp b/Engine/Plugins/Experimental/MediaViewer/Source/MediaViewer/Private/Widgets/SMediaViewerLibraryPrivate.cpp
--- a/Engine/Plugins/Experimental/MediaViewer/Source/MediaViewer/Private/Widgets/SMediaViewerLibraryPrivate.cpp
+++ b/Engine/Plugins/Experimental/MediaViewer/Source/MediaViewer/Private/Widgets/SMediaViewerLibraryPrivate.cpp
@@ -5,12 +5,13 @@
 #include "ImageViewers/ColorImageViewer.h"
 #include "IMediaViewerModule.h"
 #include "LevelEditor.h"
+#include "Library/LevelEditorViewportGroup.h"
+#include "Library/MediaTextureGroup.h"
 #include "Library/MediaViewerLibrary.h"
 #include "Library/MediaViewerLibraryGroup.h"
 #include "Library/MediaViewerLibraryIni.h"
 #include "Library/MediaViewerLibraryItem.h"
-#include "Library/MediaViewerLibraryMediaTextureGroup.h"
-#include "Library/MediaViewerLibraryTextureRenderTarget2DGroup.h"
+#include "Library/TextureRenderTarget2DGroup.h"
 #include "MediaViewer.h"
 #include "Misc/Optional.h"
 #include "Widgets/MediaViewerDelegates.h"
@@ -52,8 +53,9 @@ void SMediaViewerLibraryPrivate::Construct(const FArguments& InArgs, const TShar
 
 	TSharedRef<FMediaViewerLibrary> LibraryImpl = GetLibrary();
 
-	LibraryImpl->AddGroup(MakeShared<FMediaViewerLibraryMediaTextureGroup>(LibraryImpl));
-	LibraryImpl->AddGroup(MakeShared<FMediaViewerLibraryTextureRenderTarget2DGroup>(LibraryImpl));
+	LibraryImpl->AddGroup(MakeShared<FMediaTextureGroup>(LibraryImpl));
+	LibraryImpl->AddGroup(MakeShared<FTextureRenderTarget2DGroup>(LibraryImpl));
+	LibraryImpl->AddGroup(MakeShared<FLevelEditorViewportGroup>(LibraryImpl));
 }
 
 TSharedRef<FMediaViewerLibrary> SMediaViewerLibraryPrivate::GetLibrary() const
@@ -83,7 +85,17 @@ void SMediaViewerLibraryPrivate::OnImageViewerOpened(const TSharedRef<FMediaImag
 			return;
 		}
 
-		LibraryImpl->AddItem(Item.ToSharedRef());
+		TSharedPtr<FMediaViewerLibraryItem> ExistingItem = LibraryImpl->FindItemByValue(Item->GetItemType(), Item->GetStringValue());
+
+		if (ExistingItem)
+		{
+			Item = ExistingItem;
+			InImageViewer->UpdateId(ExistingItem->GetId());
+		}
+		else
+		{
+			LibraryImpl->AddItem(Item.ToSharedRef());
+		}
 	}
 
 	HistoryGroup->RemoveItem(Item->GetId());
diff --git a/Engine/Plugins/Experimental/MediaViewer/Source/MediaViewer/Public/ImageViewer/MediaImageViewer.h b/Engine/Plugins/Experimental/MediaViewer/Source/MediaViewer/Public/ImageViewer/MediaImageViewer.h
--- a/Engine/Plugins/Experimental/MediaViewer/Source/MediaViewer/Public/ImageViewer/MediaImageViewer.h
+++ b/Engine/Plugins/Experimental/MediaViewer/Source/MediaViewer/Public/ImageViewer/MediaImageViewer.h
@@ -179,6 +179,13 @@ public:
 	const FMediaImageViewerInfo& GetInfo() const { return ImageInfo; }
 
 	virtual bool IsValid() const { return ImageInfo.Id.IsValid(); }
+	
+	/** 
+	 * Called, for instance, when an image is dropped into the viewer and it creates an image viewer,
+	 * but the image already exists in the library. This image viewer should have its id updated to 
+	 * reflect the id already in the library.
+	 */
+	void UpdateId(const FGuid& InId);
 
 	virtual TSharedPtr<FMediaViewerLibraryItem> CreateLibraryItem() const = 0;
 
