Media Viewer: Fixed DPI Scaling

diff --git a/Engine/Plugins/Experimental/MediaViewer/Source/MediaViewer/Private/ImageViewers/MediaImageViewer.cpp b/Engine/Plugins/Experimental/MediaViewer/Source/MediaViewer/Private/ImageViewers/MediaImageViewer.cpp
--- a/Engine/Plugins/Experimental/MediaViewer/Source/MediaViewer/Private/ImageViewers/MediaImageViewer.cpp
+++ b/Engine/Plugins/Experimental/MediaViewer/Source/MediaViewer/Private/ImageViewers/MediaImageViewer.cpp
@@ -97,6 +97,11 @@ FSlateClippingZone FMediaImageViewer::CreateSlateClippingZone(const FGeometry& I
 		Zone.BottomRight.Y = Zone.TopRight.Y + Height;
 	}
 
+	Zone.TopLeft *= InGeometry.Scale;
+	Zone.TopRight *= InGeometry.Scale;
+	Zone.BottomLeft *= InGeometry.Scale;
+	Zone.BottomRight *= InGeometry.Scale;
+
 	return Zone;
 }
 
