Media Viewer: Fixed docked sidebar offset

diff --git a/Engine/Plugins/Experimental/MediaViewer/Source/MediaViewer/Private/Widgets/SMediaViewer.cpp b/Engine/Plugins/Experimental/MediaViewer/Source/MediaViewer/Private/Widgets/SMediaViewer.cpp
--- a/Engine/Plugins/Experimental/MediaViewer/Source/MediaViewer/Private/Widgets/SMediaViewer.cpp
+++ b/Engine/Plugins/Experimental/MediaViewer/Source/MediaViewer/Private/Widgets/SMediaViewer.cpp
@@ -32,6 +32,7 @@ SMediaViewer::SMediaViewer()
 	: CommandList(MakeShared<FUICommandList>())
 	, BackgroundTextureBrush(static_cast<UObject*>(nullptr), FVector2D(1), FLinearColor::White, ESlateBrushTileType::NoTile)
 	, ViewerSize(FVector2D::ZeroVector)
+	, ViewerPosition(FVector2D::ZeroVector)
 	, bInvalidated(true)
 {
 }
@@ -54,11 +55,6 @@ void SMediaViewer::Construct(const FArguments& InArgs, const FMediaViewerArgs& I
 	ContentSlot = nullptr;
 	CursorLocalPosition = {0, 0};
 
-	ViewerPosition = FVector2D(
-		MediaViewerArgs.bShowSidebar ? 25 : 0,
-		MediaViewerArgs.bShowToolbar ? 37 : 0
-	);
-
 	BindCommands();
 	CreateDelegates();
 
@@ -272,8 +268,6 @@ void SMediaViewer::Tick(const FGeometry& AllottedGeometry, const double InCurren
 int32 SMediaViewer::OnPaint(const FPaintArgs& InArgs, const FGeometry& InAllottedGeometry, const FSlateRect& InMyCullingRect,
 	FSlateWindowElementList& InOutDrawElements, int32 InLayerId, const FWidgetStyle& InWidgetStyle, bool bInParentEnabled) const
 {
-	ViewerSize = InAllottedGeometry.GetLocalSize() - ViewerPosition;
-
 	const FSlateLayoutTransform ViewerTransform = FSlateLayoutTransform(ViewerPosition);
 	const FPaintGeometry ViewerPaintGeometry = InAllottedGeometry.ToPaintGeometry(ViewerSize, ViewerTransform);
 
@@ -304,8 +298,8 @@ int32 SMediaViewer::OnPaint(const FPaintArgs& InArgs, const FGeometry& InAllotte
 		FSlateVertex BaseVert;
 		BaseVert.Color = FColor::White;
 		BaseVert.SecondaryColor = FColor::White;
-		BaseVert.Position[0] = ViewerPosition.X + 6;
-		BaseVert.Position[1] = ViewerPosition.Y + 40; // Compensate for window borders.
+		BaseVert.Position[0] = ViewerPosition.X;
+		BaseVert.Position[1] = ViewerPosition.Y;
 		BaseVert.TexCoords[0] = 0.f;
 		BaseVert.TexCoords[1] = 0.f;
 		BaseVert.TexCoords[2] = 1.f;
@@ -430,7 +424,18 @@ int32 SMediaViewer::OnPaint(const FPaintArgs& InArgs, const FGeometry& InAllotte
 
 	InOutDrawElements.PopClip();
 
-	return SCompoundWidget::OnPaint(InArgs, InAllottedGeometry, InMyCullingRect, InOutDrawElements, PaintParams.LayerId, InWidgetStyle, bInParentEnabled);
+	int32 NewLayerId = SCompoundWidget::OnPaint(InArgs, InAllottedGeometry, InMyCullingRect, InOutDrawElements, PaintParams.LayerId, InWidgetStyle, bInParentEnabled);
+
+	// Doing this after rendering or it creates flicker.
+	if (ContentSlot)
+	{
+		// InAllowedGeometry.GetAbsolutePosition() is relative to the parent.
+		// GetTickSpaceGeometry().GetAbsolutePosition() is in desktop space.
+		ViewerPosition = (ContentSlot->GetWidget()->GetTickSpaceGeometry().GetAbsolutePosition() - GetTickSpaceGeometry().GetAbsolutePosition()) / InAllottedGeometry.Scale;
+		ViewerSize = ContentSlot->GetWidget()->GetTickSpaceGeometry().GetAbsoluteSize() / InAllottedGeometry.Scale;
+	}
+
+	return NewLayerId;
 }
 
 FReply SMediaViewer::OnKeyDown(const FGeometry& InMyGeometry, const FKeyEvent& InKeyEvent)
diff --git a/Engine/Plugins/Experimental/MediaViewer/Source/MediaViewer/Private/Widgets/SMediaViewer.h b/Engine/Plugins/Experimental/MediaViewer/Source/MediaViewer/Private/Widgets/SMediaViewer.h
--- a/Engine/Plugins/Experimental/MediaViewer/Source/MediaViewer/Private/Widgets/SMediaViewer.h
+++ b/Engine/Plugins/Experimental/MediaViewer/Source/MediaViewer/Private/Widgets/SMediaViewer.h
@@ -30,6 +30,7 @@ namespace UE::MediaViewer
 		struct FMediaViewerDelegates;
 	}
 }
+
 namespace UE::MediaViewer::Private
 {
 
@@ -129,7 +130,7 @@ protected:
 	FIntPoint CursorLocalPosition;
 
 	mutable FVector2D ViewerSize;
-	FVector2D ViewerPosition;
+	mutable FVector2D ViewerPosition;
 
 	/** Do we need to be recreated next tick? */
 	bool bInvalidated;
