Media Viewer: Zoom with scrollwheel maintains mouse relative position

diff --git a/Engine/Plugins/Experimental/MediaViewer/Source/MediaViewer/Private/Widgets/MediaViewerDelegates.h b/Engine/Plugins/Experimental/MediaViewer/Source/MediaViewer/Private/Widgets/MediaViewerDelegates.h
--- a/Engine/Plugins/Experimental/MediaViewer/Source/MediaViewer/Private/Widgets/MediaViewerDelegates.h
+++ b/Engine/Plugins/Experimental/MediaViewer/Source/MediaViewer/Private/Widgets/MediaViewerDelegates.h
@@ -93,6 +93,9 @@ struct FMediaViewerDelegates
 	/** Multiplies the scale of all image viewers. */
 	FMultiplyScale MultiplyScaleToAll;
 
+	/** Multiplies the scale of all image viewers. */
+	FMultiplyScale MultiplyScaleAroundCursorToAll;
+
 	/** Sets the transform of all image viewers. */
 	FSetTransform SetTransformToAll;
 
diff --git a/Engine/Plugins/Experimental/MediaViewer/Source/MediaViewer/Private/Widgets/SMediaImageViewerOverlay.cpp b/Engine/Plugins/Experimental/MediaViewer/Source/MediaViewer/Private/Widgets/SMediaImageViewerOverlay.cpp
--- a/Engine/Plugins/Experimental/MediaViewer/Source/MediaViewer/Private/Widgets/SMediaImageViewerOverlay.cpp
+++ b/Engine/Plugins/Experimental/MediaViewer/Source/MediaViewer/Private/Widgets/SMediaImageViewerOverlay.cpp
@@ -11,6 +11,8 @@
 #include "IMediaViewerModule.h"
 #include "Library/MediaViewerLibrary.h"
 #include "Library/MediaViewerLibraryItem.h"
+#include "Math/Vector.h"
+#include "Math/Vector2D.h"
 #include "MediaViewer.h"
 #include "MediaViewerCommands.h"
 #include "SAssetDropTarget.h"
@@ -100,25 +102,37 @@ const TSharedRef<FUICommandList>& SMediaImageViewerOverlay::GetCommandList() con
 }
 
 FIntPoint SMediaImageViewerOverlay::GetImageViewerPixelCoordinates() const
+{
+	const FVector2f CursorLocation = GetImageViewerPixelCoordinates_Exact();
+
+	return FIntPoint(
+		FMath::RoundToInt(CursorLocation.X),
+		FMath::RoundToInt(CursorLocation.Y)
+	);
+}
+
+FVector2f SMediaImageViewerOverlay::GetImageViewerPixelCoordinates_Exact() const
 {
 	TSharedPtr<FMediaImageViewer> ImageViewer = GetImageViewerDelegate.Execute(Position);
 
 	if (!ImageViewer.IsValid())
 	{
-		return FIntPoint(-1, -1);
+		return FVector2f(-1, -1);
 	}
 
+	const FIntPoint ViewerPosition = Delegates->GetViewerPosition.Execute();
 	const FIntPoint ViewerSize = Delegates->GetViewerSize.Execute();
-	const FIntPoint CursorLocation = Delegates->GetCursorLocation.Execute() - Delegates->GetViewerPosition.Execute();
+	const FIntPoint CursorLocation = Delegates->GetCursorLocation.Execute() - ViewerPosition;
 	const FVector2D TopLeft = ImageViewer->GetPaintOffset(ViewerSize, Delegates->GetViewerPosition.Execute());
 
-	FVector2D CursorLocationD = FVector2D(CursorLocation.X, CursorLocation.Y) - TopLeft + Delegates->GetViewerPosition.Execute();
-	CursorLocationD /= ImageViewer->GetPaintSettings().Scale;
+	FVector2f CursorLocationF = FVector2f(CursorLocation.X, CursorLocation.Y);
+	CursorLocationF.X -= TopLeft.X;
+	CursorLocationF.Y -= TopLeft.Y;
+	CursorLocationF.X += ViewerPosition.X;
+	CursorLocationF.Y += ViewerPosition.Y;
+	CursorLocationF /= ImageViewer->GetPaintSettings().Scale;
 
-	return FIntPoint(
-		FMath::RoundToInt(CursorLocationD.X),
-		FMath::RoundToInt(CursorLocationD.Y)
-	);
+	return CursorLocationF;
 }
 
 bool SMediaImageViewerOverlay::IsCursorOverImageViewer() const
@@ -243,15 +257,17 @@ FReply SMediaImageViewerOverlay::OnMouseWheel(const FGeometry& InMyGeometry, con
 			? SlowMultiplier
 			: NormalMultiplier);
 	
-	const float WheelDelta = InMouseEvent.GetWheelDelta();
+	const float ScaleMultiplier = InMouseEvent.GetWheelDelta() > 0
+		? ModifierMultiplier
+		: (1.f / ModifierMultiplier);
 
-	if (WheelDelta > 0)
+	if (Delegates->AreTransformsLocked.Execute())
 	{
-		Try_SetScale(GetScale() * ModifierMultiplier);
+		Delegates->MultiplyScaleAroundCursorToAll.Execute(ScaleMultiplier);
 	}
-	else if (WheelDelta < 0)
+	else
 	{
-		Try_SetScale(GetScale() / ModifierMultiplier);
+		MultiplyScaleAroundCursor(ScaleMultiplier);
 	}
 
 	return SCompoundWidget::OnMouseWheel(InMyGeometry, InMouseEvent);
@@ -619,6 +635,27 @@ void SMediaImageViewerOverlay::ToggleOverlay()
 	}
 }
 
+void SMediaImageViewerOverlay::MultiplyScaleAroundCursor(float InMultipler)
+{
+	const FVector2f CursorLocationBefore = GetImageViewerPixelCoordinates_Exact();
+
+	SetScale(GetScale() * InMultipler);
+
+	const FVector2f CursorLocationAfter = GetImageViewerPixelCoordinates_Exact();
+
+	if (!(CursorLocationBefore - CursorLocationAfter).IsNearlyZero())
+	{
+		FVector Offset;
+		Offset.X = CursorLocationAfter.X - CursorLocationBefore.X;
+		Offset.Y = CursorLocationAfter.Y - CursorLocationBefore.Y;
+		Offset.Z = 0;
+
+		Offset *= GetScale();
+
+		AddOffset(Offset);
+	}
+}
+
 void SMediaImageViewerOverlay::ResetTransform()
 {
 	SetOffset(FVector::ZeroVector);
diff --git a/Engine/Plugins/Experimental/MediaViewer/Source/MediaViewer/Private/Widgets/SMediaImageViewerOverlay.h b/Engine/Plugins/Experimental/MediaViewer/Source/MediaViewer/Private/Widgets/SMediaImageViewerOverlay.h
--- a/Engine/Plugins/Experimental/MediaViewer/Source/MediaViewer/Private/Widgets/SMediaImageViewerOverlay.h
+++ b/Engine/Plugins/Experimental/MediaViewer/Source/MediaViewer/Private/Widgets/SMediaImageViewerOverlay.h
@@ -45,6 +45,8 @@ public:
 
 	FIntPoint GetImageViewerPixelCoordinates() const;
 
+	FVector2f GetImageViewerPixelCoordinates_Exact() const;
+
 	bool IsCursorOverImageViewer() const;
 
 	void UpdateMouse(const TOptional<FIntPoint>& InMousePosition);
@@ -52,6 +54,8 @@ public:
 	bool IsOverlayEnabled() const;
 	void ToggleOverlay();
 
+	void MultiplyScaleAroundCursor(float InMultipler);
+
 	//~ Begin SWidget
 	virtual FCursorReply OnCursorQuery(const FGeometry& InMyGeometry, const FPointerEvent& InCursorEvent) const override;
 	virtual FReply OnKeyDown(const FGeometry& InMyGeometry, const FKeyEvent& InKeyEvent) override;
diff --git a/Engine/Plugins/Experimental/MediaViewer/Source/MediaViewer/Private/Widgets/SMediaViewer.cpp b/Engine/Plugins/Experimental/MediaViewer/Source/MediaViewer/Private/Widgets/SMediaViewer.cpp
--- a/Engine/Plugins/Experimental/MediaViewer/Source/MediaViewer/Private/Widgets/SMediaViewer.cpp
+++ b/Engine/Plugins/Experimental/MediaViewer/Source/MediaViewer/Private/Widgets/SMediaViewer.cpp
@@ -481,6 +481,7 @@ void SMediaViewer::CreateDelegates()
 	Delegates->AddOffsetToAll.BindSP(this, &SMediaViewer::AddOffsetToAll);
 	Delegates->AddRotationToAll.BindSP(this, &SMediaViewer::AddRotationToAll);
 	Delegates->MultiplyScaleToAll.BindSP(this, &SMediaViewer::MultiplyScaleToAll);
+	Delegates->MultiplyScaleAroundCursorToAll.BindSP(this, &SMediaViewer::MultiplyScaleAroundCursorToAll);
 	Delegates->SetTransformToAll.BindSP(this, &SMediaViewer::SetTransformToAll);
 	Delegates->ResetTransformToAll.BindSP(this, &SMediaViewer::ResetTransformToAll);
 	Delegates->GetSecondImageViewerOpacity.BindSPLambda(this, [this] { return MediaViewerSettings.SecondImageOpacity; });
@@ -764,6 +765,17 @@ void SMediaViewer::MultiplyScaleToAll(float InMultiple)
 	}
 }
 
+void SMediaViewer::MultiplyScaleAroundCursorToAll(float InMultiple)
+{
+	for (int32 Index = 0; Index < static_cast<int32>(EMediaImageViewerPosition::COUNT); ++Index)
+	{
+		if (TSharedPtr<SMediaImageViewerOverlay> Overlay = GetOverlay(static_cast<EMediaImageViewerPosition>(Index)))
+		{
+			Overlay->MultiplyScaleAroundCursor(InMultiple);
+		}
+	}
+}
+
 void SMediaViewer::SetTransformToAll(const FVector& InOffset, const FRotator& InRotation, float InScale)
 {
 	for (const TSharedPtr<FMediaImageViewer>& ImageViewer : ImageViewers)
diff --git a/Engine/Plugins/Experimental/MediaViewer/Source/MediaViewer/Private/Widgets/SMediaViewer.h b/Engine/Plugins/Experimental/MediaViewer/Source/MediaViewer/Private/Widgets/SMediaViewer.h
--- a/Engine/Plugins/Experimental/MediaViewer/Source/MediaViewer/Private/Widgets/SMediaViewer.h
+++ b/Engine/Plugins/Experimental/MediaViewer/Source/MediaViewer/Private/Widgets/SMediaViewer.h
@@ -186,6 +186,9 @@ protected:
 	/** Multiplays the scale of all image viewers by a value. */
 	void MultiplyScaleToAll(float InMultiple);
 
+	/** Multiplays the scale of all image viewers by a value keeping the cursor in the same relative position for each image. */
+	void MultiplyScaleAroundCursorToAll(float InMultiple);
+
 	/** Sets the transform of all image viewers to the given values. */
 	void SetTransformToAll(const FVector& InOffset, const FRotator& InRotation, float InScale);
 
