Media Viewer: Drag and drop upgrade
- You can now drag and drop from the library to the viewers.
- Added dedicated drop target class to replace asset drop target usage.
- Fixed crash bug with library items.
- Fixed positioning issue when the sidebar is docked.
- Added SetImageViewer delegate to avoid module usage.

diff --git a/Engine/Plugins/Experimental/MediaViewer/Source/MediaViewer/Private/ImageViewers/ColorImageViewer.cpp b/Engine/Plugins/Experimental/MediaViewer/Source/MediaViewer/Private/ImageViewers/ColorImageViewer.cpp
--- a/Engine/Plugins/Experimental/MediaViewer/Source/MediaViewer/Private/ImageViewers/ColorImageViewer.cpp
+++ b/Engine/Plugins/Experimental/MediaViewer/Source/MediaViewer/Private/ImageViewers/ColorImageViewer.cpp
@@ -82,7 +82,9 @@ TSharedPtr<FSlateBrush> FColorImageViewer::FItem::CreateThumbnail()
 		return nullptr;
 	}
 
-	return MakeShared<FSlateColorBrush>(LoadFromString(StringValue));
+	TSharedRef<FSlateColorBrush> ThumbnailBrush = MakeShared<FSlateColorBrush>(LoadFromString(StringValue));
+
+	return ThumbnailBrush;
 }
 
 TSharedPtr<FMediaImageViewer> FColorImageViewer::FItem::CreateImageViewer() const
diff --git a/Engine/Plugins/Experimental/MediaViewer/Source/MediaViewer/Private/ImageViewers/LevelEditorViewportImageViewer.cpp b/Engine/Plugins/Experimental/MediaViewer/Source/MediaViewer/Private/ImageViewers/LevelEditorViewportImageViewer.cpp
--- a/Engine/Plugins/Experimental/MediaViewer/Source/MediaViewer/Private/ImageViewers/LevelEditorViewportImageViewer.cpp
+++ b/Engine/Plugins/Experimental/MediaViewer/Source/MediaViewer/Private/ImageViewers/LevelEditorViewportImageViewer.cpp
@@ -161,7 +161,9 @@ TSharedPtr<FSlateBrush> FLevelEditorViewportImageViewer::FItem::CreateThumbnail(
 
 	RenderViewport(TNotNull<FViewport*>(Viewport.Get()), TNotNull<UTextureRenderTarget2D*>(RenderTarget), /* Resize render target */ false);
 
-	return MakeShared<FSlateImageBrush>(RenderTarget, FVector2D(RenderTarget->GetSurfaceWidth(), RenderTarget->GetSurfaceDepth()));
+	TSharedRef<FSlateImageBrush> ThumbnailBrush = MakeShared<FSlateImageBrush>(RenderTarget, FVector2D(RenderTarget->GetSurfaceWidth(), RenderTarget->GetSurfaceDepth()));
+
+	return ThumbnailBrush;
 }
 
 TSharedPtr<FMediaImageViewer> FLevelEditorViewportImageViewer::FItem::CreateImageViewer() const
diff --git a/Engine/Plugins/Experimental/MediaViewer/Source/MediaViewer/Private/ImageViewers/MaterialInterfaceImageViewer.cpp b/Engine/Plugins/Experimental/MediaViewer/Source/MediaViewer/Private/ImageViewers/MaterialInterfaceImageViewer.cpp
--- a/Engine/Plugins/Experimental/MediaViewer/Source/MediaViewer/Private/ImageViewers/MaterialInterfaceImageViewer.cpp
+++ b/Engine/Plugins/Experimental/MediaViewer/Source/MediaViewer/Private/ImageViewers/MaterialInterfaceImageViewer.cpp
@@ -160,7 +160,9 @@ TSharedPtr<FSlateBrush> FMaterialInterfaceImageViewer::FItem::CreateThumbnail()
 		CreateRenderTargetThumbnail();
 	}
 
-	return MakeShared<FSlateImageBrush>(RenderTarget, FVector2D(RenderTarget->GetSurfaceWidth(), RenderTarget->GetSurfaceDepth()));
+	TSharedRef<FSlateImageBrush> ThumbnailBrush = MakeShared<FSlateImageBrush>(RenderTarget, FVector2D(RenderTarget->GetSurfaceWidth(), RenderTarget->GetSurfaceDepth()));
+
+	return ThumbnailBrush;
 }
 
 TSharedPtr<FMediaImageViewer> FMaterialInterfaceImageViewer::FItem::CreateImageViewer() const
diff --git a/Engine/Plugins/Experimental/MediaViewer/Source/MediaViewer/Private/ImageViewers/MaterialInterfaceImageViewer.h b/Engine/Plugins/Experimental/MediaViewer/Source/MediaViewer/Private/ImageViewers/MaterialInterfaceImageViewer.h
--- a/Engine/Plugins/Experimental/MediaViewer/Source/MediaViewer/Private/ImageViewers/MaterialInterfaceImageViewer.h
+++ b/Engine/Plugins/Experimental/MediaViewer/Source/MediaViewer/Private/ImageViewers/MaterialInterfaceImageViewer.h
@@ -8,7 +8,6 @@
 #include "UObject/GCObject.h"
 
 #include "ImageViewers/TextureMipCache.h"
-#include "IMediaViewerModule.h"
 #include "Templates/SharedPointer.h"
 
 #include "MaterialInterfaceImageViewer.generated.h"
diff --git a/Engine/Plugins/Experimental/MediaViewer/Source/MediaViewer/Private/ImageViewers/MediaImageViewer.cpp b/Engine/Plugins/Experimental/MediaViewer/Source/MediaViewer/Private/ImageViewers/MediaImageViewer.cpp
--- a/Engine/Plugins/Experimental/MediaViewer/Source/MediaViewer/Private/ImageViewers/MediaImageViewer.cpp
+++ b/Engine/Plugins/Experimental/MediaViewer/Source/MediaViewer/Private/ImageViewers/MediaImageViewer.cpp
@@ -15,19 +15,9 @@ const FSlateColorBrush FMediaImageViewer::BackgroundColorBrush(FLinearColor::Whi
 
 FMediaImageViewer::FMediaImageViewer(const FMediaImageViewerInfo& InImageInfo)
 	: ImageInfo(InImageInfo)
-	, PanelSettings(MakeShared<TStructOnScope<FMediaImageViewerPanelSettings>>())
-	, PaintSettings(MakeShared<TStructOnScope<FMediaImagePaintSettings>>())
 	, BackgroundImageBrush(static_cast<UObject*>(nullptr), FVector2D(1))
 	, Brush(nullptr)
 {
-	PanelSettings->InitializeAs<FMediaImageViewerPanelSettings>();
-	PaintSettings->InitializeAs<FMediaImagePaintSettings>();
-}
-
-FMediaImageViewer::~FMediaImageViewer()
-{
-	static_cast<FStructOnScope&>(PanelSettings.Get()).Destroy();
-	static_cast<FStructOnScope&>(PaintSettings.Get()).Destroy();
 }
 
 void FMediaImageViewer::UpdateId(const FGuid& InId)
@@ -37,7 +27,7 @@ void FMediaImageViewer::UpdateId(const FGuid& InId)
 
 void FMediaImageViewer::NotifyPostChange(const FPropertyChangedEvent& InPropertyChangedEvent, FProperty* InPropertyThatChanged)
 {
-	if (UTexture* Texture = GetPanelSettings().BackgroundTexture.LoadSynchronous())
+	if (UTexture* Texture = PanelSettings.BackgroundTexture.LoadSynchronous())
 	{
 		BackgroundImageBrush.SetResourceObject(Texture);
 		BackgroundImageBrush.ImageSize.X = Texture->GetSurfaceWidth();
@@ -137,7 +127,7 @@ void FMediaImageViewer::PaintPanel(FMediaImagePaintParams& InPaintParams, const
 {
 	const FLinearColor AlphaTint = FLinearColor(1.f, 1.f, 1.f, InPaintParams.ImageOpacity);
 
-	if (GetPanelSettings().BackgroundColor.IsSet())
+	if (PanelSettings.BackgroundColor.IsSet())
 	{
 		FSlateDrawElement::MakeBox(
 			InPaintParams.DrawElements,
@@ -145,7 +135,7 @@ void FMediaImageViewer::PaintPanel(FMediaImagePaintParams& InPaintParams, const
 			InPaintGeometry.Geometry,
 			&BackgroundColorBrush,
 			ESlateDrawEffect::NoPixelSnapping,
-			GetPanelSettings().BackgroundColor.GetValue() * AlphaTint
+			PanelSettings.BackgroundColor.GetValue() * AlphaTint
 		);
 
 		++InPaintParams.LayerId;
diff --git a/Engine/Plugins/Experimental/MediaViewer/Source/MediaViewer/Private/ImageViewers/MediaTextureImageViewer.cpp b/Engine/Plugins/Experimental/MediaViewer/Source/MediaViewer/Private/ImageViewers/MediaTextureImageViewer.cpp
--- a/Engine/Plugins/Experimental/MediaViewer/Source/MediaViewer/Private/ImageViewers/MediaTextureImageViewer.cpp
+++ b/Engine/Plugins/Experimental/MediaViewer/Source/MediaViewer/Private/ImageViewers/MediaTextureImageViewer.cpp
@@ -125,7 +125,9 @@ TSharedPtr<FSlateBrush> FMediaTextureImageViewer::FItem::CreateThumbnail()
 		return nullptr;
 	}
 
-	return MakeShared<FSlateImageBrush>(MediaTexture, FVector2D(MediaTexture->GetSurfaceWidth(), MediaTexture->GetSurfaceDepth()));
+	TSharedRef<FSlateImageBrush> ThumbnailBrush = MakeShared<FSlateImageBrush>(MediaTexture, FVector2D(MediaTexture->GetSurfaceWidth(), MediaTexture->GetSurfaceDepth()));
+
+	return ThumbnailBrush;
 }
 
 TSharedPtr<FMediaImageViewer> FMediaTextureImageViewer::FItem::CreateImageViewer() const
diff --git a/Engine/Plugins/Experimental/MediaViewer/Source/MediaViewer/Private/ImageViewers/MediaTextureImageViewer.h b/Engine/Plugins/Experimental/MediaViewer/Source/MediaViewer/Private/ImageViewers/MediaTextureImageViewer.h
--- a/Engine/Plugins/Experimental/MediaViewer/Source/MediaViewer/Private/ImageViewers/MediaTextureImageViewer.h
+++ b/Engine/Plugins/Experimental/MediaViewer/Source/MediaViewer/Private/ImageViewers/MediaTextureImageViewer.h
@@ -7,7 +7,6 @@
 #include "Library/MediaViewerLibraryItem.h"
 #include "UObject/GCObject.h"
 
-#include "IMediaViewerModule.h"
 #include "Templates/SharedPointer.h"
 
 #include "MediaTextureImageViewer.generated.h"
diff --git a/Engine/Plugins/Experimental/MediaViewer/Source/MediaViewer/Private/ImageViewers/Texture2DImageViewer.cpp b/Engine/Plugins/Experimental/MediaViewer/Source/MediaViewer/Private/ImageViewers/Texture2DImageViewer.cpp
--- a/Engine/Plugins/Experimental/MediaViewer/Source/MediaViewer/Private/ImageViewers/Texture2DImageViewer.cpp
+++ b/Engine/Plugins/Experimental/MediaViewer/Source/MediaViewer/Private/ImageViewers/Texture2DImageViewer.cpp
@@ -121,14 +121,19 @@ TSharedPtr<FSlateBrush> FTexture2DImageViewer::FItem::CreateThumbnail()
 		return nullptr;
 	}
 
-	UTexture2D* Texture = LoadObject<UTexture2D>(GetTransientPackage(), *StringValue);
-
 	if (!Texture)
 	{
-		return nullptr;
+		Texture = LoadObject<UTexture2D>(GetTransientPackage(), *StringValue);
+
+		if (!Texture)
+		{
+			return nullptr;
+		}
 	}
 
-	return MakeShared<FSlateImageBrush>(Texture, FVector2D(Texture->GetSurfaceWidth(), Texture->GetSurfaceDepth()));
+	TSharedRef<FSlateImageBrush> ThumbnailBrush = MakeShared<FSlateImageBrush>(Texture, FVector2D(Texture->GetSurfaceWidth(), Texture->GetSurfaceDepth()));
+
+	return ThumbnailBrush;
 }
 
 TSharedPtr<FMediaImageViewer> FTexture2DImageViewer::FItem::CreateImageViewer() const
@@ -138,11 +143,14 @@ TSharedPtr<FMediaImageViewer> FTexture2DImageViewer::FItem::CreateImageViewer()
 		return nullptr;
 	}
 
-	UTexture2D* Texture = LoadAssetFromString<UTexture2D>(StringValue);
-
 	if (!Texture)
 	{
-		return nullptr;
+		Texture = LoadAssetFromString<UTexture2D>(StringValue);
+
+		if (!Texture)
+		{
+			return nullptr;
+		}
 	}
 
 	if (Id.IsValid())
@@ -163,6 +171,20 @@ TSharedPtr<FMediaViewerLibraryItem> FTexture2DImageViewer::FItem::Clone() const
 	return MakeShared<FTexture2DImageViewer::FItem>(FPrivateToken(), *this);
 }
 
+FString FTexture2DImageViewer::FItem::GetReferencerName() const
+{
+	static const FString ReferencerName = TEXT("FTexture2DImageViewer");
+	return ReferencerName;
+}
+
+void FTexture2DImageViewer::FItem::AddReferencedObjects(FReferenceCollector& Collector)
+{
+	if (Texture)
+	{
+		Collector.AddReferencedObject(Texture);
+	}
+}
+
 FTexture2DImageViewer::FTexture2DImageViewer(TNotNull<UTexture2D*> InTexture)
 	: FTexture2DImageViewer(FGuid::NewGuid(), InTexture)
 {
diff --git a/Engine/Plugins/Experimental/MediaViewer/Source/MediaViewer/Private/ImageViewers/Texture2DImageViewer.h b/Engine/Plugins/Experimental/MediaViewer/Source/MediaViewer/Private/ImageViewers/Texture2DImageViewer.h
--- a/Engine/Plugins/Experimental/MediaViewer/Source/MediaViewer/Private/ImageViewers/Texture2DImageViewer.h
+++ b/Engine/Plugins/Experimental/MediaViewer/Source/MediaViewer/Private/ImageViewers/Texture2DImageViewer.h
@@ -8,7 +8,6 @@
 #include "UObject/GCObject.h"
 
 #include "ImageViewers/TextureMipCache.h"
-#include "IMediaViewerModule.h"
 #include "Templates/SharedPointer.h"
 
 #include "Texture2DImageViewer.generated.h"
@@ -49,8 +48,10 @@ public:
 		virtual TSharedPtr<FMediaViewerLibraryItem> CreateLibraryItem(const FMediaViewerLibraryItem& InSavedItem) const override;
 	};
 
-	struct FItem : public FMediaViewerLibraryItem
+	struct FItem : public FMediaViewerLibraryItem, public FGCObject
 	{
+		mutable TObjectPtr<UTexture2D> Texture;
+
 		FItem(const FText& InName, const FText& InToolTip, bool bInTransient, TNotNull<UTexture2D*> InTexture);
 
 		FItem(const FGuid& InId, const FText& InName, const FText& InToolTip, bool bInTransient, TNotNull<UTexture2D*> InTexture);
@@ -66,6 +67,11 @@ public:
 		virtual TSharedPtr<FMediaImageViewer> CreateImageViewer() const override;
 		virtual TSharedPtr<FMediaViewerLibraryItem> Clone() const override;
 		//~ End FMediaViewerLibraryItem
+
+		//~ Begin FGCObject
+		virtual FString GetReferencerName() const override;
+		virtual void AddReferencedObjects(FReferenceCollector& Collector) override;
+		//~ End FGCObject
 	};
 
 	static const FLazyName ItemTypeName;
diff --git a/Engine/Plugins/Experimental/MediaViewer/Source/MediaViewer/Private/ImageViewers/TextureRenderTarget2DImageViewer.cpp b/Engine/Plugins/Experimental/MediaViewer/Source/MediaViewer/Private/ImageViewers/TextureRenderTarget2DImageViewer.cpp
--- a/Engine/Plugins/Experimental/MediaViewer/Source/MediaViewer/Private/ImageViewers/TextureRenderTarget2DImageViewer.cpp
+++ b/Engine/Plugins/Experimental/MediaViewer/Source/MediaViewer/Private/ImageViewers/TextureRenderTarget2DImageViewer.cpp
@@ -122,7 +122,9 @@ TSharedPtr<FSlateBrush> FTextureRenderTarget2DImageViewer::FItem::CreateThumbnai
 		return nullptr;
 	}
 
-	return MakeShared<FSlateImageBrush>(RenderTarget, FVector2D(RenderTarget->GetSurfaceWidth(), RenderTarget->GetSurfaceDepth()));
+	TSharedRef<FSlateImageBrush> ThumbnailBrush = MakeShared<FSlateImageBrush>(RenderTarget, FVector2D(RenderTarget->GetSurfaceWidth(), RenderTarget->GetSurfaceDepth()));
+
+	return ThumbnailBrush;
 }
 
 TSharedPtr<FMediaImageViewer> FTextureRenderTarget2DImageViewer::FItem::CreateImageViewer() const
diff --git a/Engine/Plugins/Experimental/MediaViewer/Source/MediaViewer/Private/ImageViewers/TextureRenderTarget2DImageViewer.h b/Engine/Plugins/Experimental/MediaViewer/Source/MediaViewer/Private/ImageViewers/TextureRenderTarget2DImageViewer.h
--- a/Engine/Plugins/Experimental/MediaViewer/Source/MediaViewer/Private/ImageViewers/TextureRenderTarget2DImageViewer.h
+++ b/Engine/Plugins/Experimental/MediaViewer/Source/MediaViewer/Private/ImageViewers/TextureRenderTarget2DImageViewer.h
@@ -7,7 +7,6 @@
 #include "Library/MediaViewerLibraryItem.h"
 #include "UObject/GCObject.h"
 
-#include "IMediaViewerModule.h"
 #include "Templates/SharedPointer.h"
 
 #include "TextureRenderTarget2DImageViewer.generated.h"
diff --git a/Engine/Plugins/Experimental/MediaViewer/Source/MediaViewer/Private/ImageViewers/ViewportImageViewer.h b/Engine/Plugins/Experimental/MediaViewer/Source/MediaViewer/Private/ImageViewers/ViewportImageViewer.h
--- a/Engine/Plugins/Experimental/MediaViewer/Source/MediaViewer/Private/ImageViewers/ViewportImageViewer.h
+++ b/Engine/Plugins/Experimental/MediaViewer/Source/MediaViewer/Private/ImageViewers/ViewportImageViewer.h
@@ -7,7 +7,6 @@
 #include "UObject/GCObject.h"
 
 #include "ImageViewers/TextureMipCache.h"
-#include "IMediaViewerModule.h"
 #include "Templates/SharedPointer.h"
 
 #include "ViewportImageViewer.generated.h"
diff --git a/Engine/Plugins/Experimental/MediaViewer/Source/MediaViewer/Private/MediaViewerModule.cpp b/Engine/Plugins/Experimental/MediaViewer/Source/MediaViewer/Private/MediaViewerModule.cpp
--- a/Engine/Plugins/Experimental/MediaViewer/Source/MediaViewer/Private/MediaViewerModule.cpp
+++ b/Engine/Plugins/Experimental/MediaViewer/Source/MediaViewer/Private/MediaViewerModule.cpp
@@ -19,6 +19,7 @@
 #include "Misc/NotNull.h"
 #include "Modules/ModuleManager.h"
 #include "Widgets/Docking/SDockTab.h"
+#include "Widgets/MediaViewerDelegates.h"
 #include "Widgets/SMediaViewerLibrary.h"
 #include "Widgets/SMediaViewerTab.h"
 #include "WorkspaceMenuStructure.h"
@@ -203,7 +204,7 @@ TSharedRef<SMediaViewerTab> FMediaViewerModule::CreateMediaViewer(const FMediaVi
 
 TSharedRef<IMediaViewerLibraryWidget> FMediaViewerModule::CreateLibrary(const IMediaViewerLibraryWidget::FArgs& InArgs)
 {
-	return SNew(SMediaViewerLibrary, InArgs);
+	return SNew(SMediaViewerLibrary, InArgs, MakeShared<FMediaViewerDelegates>());
 }
 
 bool FMediaViewerModule::HasImage(EMediaImageViewerPosition InPosition) const
diff --git a/Engine/Plugins/Experimental/MediaViewer/Source/MediaViewer/Private/Widgets/MediaViewerDelegates.h b/Engine/Plugins/Experimental/MediaViewer/Source/MediaViewer/Private/Widgets/MediaViewerDelegates.h
--- a/Engine/Plugins/Experimental/MediaViewer/Source/MediaViewer/Private/Widgets/MediaViewerDelegates.h
+++ b/Engine/Plugins/Experimental/MediaViewer/Source/MediaViewer/Private/Widgets/MediaViewerDelegates.h
@@ -42,6 +42,7 @@ struct FMediaViewerDelegates
 	DECLARE_DELEGATE_RetVal(TSharedPtr<FUICommandList>, FGetCommandList)
 
 	DECLARE_DELEGATE_RetVal_OneParam(TSharedPtr<FMediaImageViewer>, FGetImageViewer, EMediaImageViewerPosition)
+	DECLARE_DELEGATE_TwoParams(FSetImageViewer, EMediaImageViewerPosition, const TSharedRef<FMediaImageViewer>&)
 	DECLARE_DELEGATE_OneParam(FSimpleForPosition, EMediaImageViewerPosition)
 	DECLARE_DELEGATE_RetVal_OneParam(FIntPoint, FGetLocationForPosition, EMediaImageViewerPosition)
 	DECLARE_DELEGATE_RetVal_OneParam(bool, FGetBoolForPosition, EMediaImageViewerPosition)
@@ -130,6 +131,9 @@ struct FMediaViewerDelegates
 	/** Gets the image viewer in the given position. */
 	FGetImageViewer GetImageViewer;
 
+	/** Sets the image viewer in the given position. */
+	FSetImageViewer SetImageViewer;
+
 	/** Gets the pixel coordinates hovered by the mouse for the image viewer in the given position. */
 	FGetLocationForPosition GetPixelCoordinates;
 
diff --git a/Engine/Plugins/Experimental/MediaViewer/Source/MediaViewer/Private/Widgets/MediaViewerLibraryItemDragDropOperation.cpp b/Engine/Plugins/Experimental/MediaViewer/Source/MediaViewer/Private/Widgets/MediaViewerLibraryItemDragDropOperation.cpp
--- a/Engine/Plugins/Experimental/MediaViewer/Source/MediaViewer/Private/Widgets/MediaViewerLibraryItemDragDropOperation.cpp
+++ b/Engine/Plugins/Experimental/MediaViewer/Source/MediaViewer/Private/Widgets/MediaViewerLibraryItemDragDropOperation.cpp
@@ -3,7 +3,6 @@
 #include "Widgets/MediaViewerLibraryItemDragDropOperation.h"
 
 #include "DetailLayoutBuilder.h"
-#include "IMediaViewerModule.h"
 #include "Library/MediaViewerLibrary.h"
 #include "Library/MediaViewerLibraryGroup.h"
 #include "Library/MediaViewerLibraryItem.h"
diff --git a/Engine/Plugins/Experimental/MediaViewer/Source/MediaViewer/Private/Widgets/SMediaImageViewerDetails.cpp b/Engine/Plugins/Experimental/MediaViewer/Source/MediaViewer/Private/Widgets/SMediaImageViewerDetails.cpp
--- a/Engine/Plugins/Experimental/MediaViewer/Source/MediaViewer/Private/Widgets/SMediaImageViewerDetails.cpp
+++ b/Engine/Plugins/Experimental/MediaViewer/Source/MediaViewer/Private/Widgets/SMediaImageViewerDetails.cpp
@@ -23,12 +23,12 @@ void SMediaImageViewerDetails::PrivateRegisterAttributes(FSlateAttributeDescript
 }
 
 void SMediaImageViewerDetails::Construct(const FArguments& InArgs, EMediaImageViewerPosition InPosition, 
-	const TSharedPtr<FMediaViewerDelegates>& InDelegates)
+	const TSharedRef<FMediaViewerDelegates>& InDelegates)
 {
 	Position = InPosition;
-	GetImageViewerDelegate = InDelegates->GetImageViewer;
+	Delegates = InDelegates;
 
-	TSharedPtr<FMediaImageViewer> ImageViewer = GetImageViewerDelegate.Execute(Position);
+	TSharedPtr<FMediaImageViewer> ImageViewer = Delegates->GetImageViewer.Execute(Position);
 
 	if (!ImageViewer.IsValid())
 	{
@@ -73,7 +73,7 @@ void SMediaImageViewerDetails::Construct(const FArguments& InArgs, EMediaImageVi
 
 TSharedRef<SWidget> SMediaImageViewerDetails::CreatePanelSettings()
 {
-	TSharedPtr<FMediaImageViewer> ImageViewer = GetImageViewerDelegate.Execute(Position);
+	TSharedPtr<FMediaImageViewer> ImageViewer = Delegates->GetImageViewer.Execute(Position);
 
 	if (!ImageViewer.IsValid())
 	{
@@ -81,7 +81,7 @@ TSharedRef<SWidget> SMediaImageViewerDetails::CreatePanelSettings()
 	}
 
 	PanelDetailsView = FMediaViewerUtils::CreateStructDetailsView(
-		ImageViewer->GetPanelSettingsOnScope(),
+		MakeShared<FStructOnScope>(FMediaImageViewerPanelSettings::StaticStruct(), reinterpret_cast<uint8*>(&ImageViewer->GetPanelSettings())),
 		LOCTEXT("PanelSettings", "Panel"),
 		ImageViewer.Get()
 	);
@@ -91,7 +91,7 @@ TSharedRef<SWidget> SMediaImageViewerDetails::CreatePanelSettings()
 
 TSharedRef<SWidget> SMediaImageViewerDetails::CreatePaintSettings()
 {
-	TSharedPtr<FMediaImageViewer> ImageViewer = GetImageViewerDelegate.Execute(Position);
+	TSharedPtr<FMediaImageViewer> ImageViewer = Delegates->GetImageViewer.Execute(Position);
 
 	if (!ImageViewer.IsValid())
 	{
@@ -99,7 +99,7 @@ TSharedRef<SWidget> SMediaImageViewerDetails::CreatePaintSettings()
 	}
 
 	PaintDetailsView = FMediaViewerUtils::CreateStructDetailsView(
-		ImageViewer->GetPaintSettingsOnScope(),
+		MakeShared<FStructOnScope>(FMediaImagePaintSettings::StaticStruct(), reinterpret_cast<uint8*>(&ImageViewer->GetPaintSettings())),
 		LOCTEXT("PaintSettings", "Media (Base)"),
 		ImageViewer.Get()
 	);
@@ -109,7 +109,7 @@ TSharedRef<SWidget> SMediaImageViewerDetails::CreatePaintSettings()
 
 TSharedPtr<SWidget> SMediaImageViewerDetails::CreateCustomSettings()
 {
-	TSharedPtr<FMediaImageViewer> ImageViewer = GetImageViewerDelegate.Execute(Position);
+	TSharedPtr<FMediaImageViewer> ImageViewer = Delegates->GetImageViewer.Execute(Position);
 
 	if (!ImageViewer.IsValid())
 	{
diff --git a/Engine/Plugins/Experimental/MediaViewer/Source/MediaViewer/Private/Widgets/SMediaImageViewerDetails.h b/Engine/Plugins/Experimental/MediaViewer/Source/MediaViewer/Private/Widgets/SMediaImageViewerDetails.h
--- a/Engine/Plugins/Experimental/MediaViewer/Source/MediaViewer/Private/Widgets/SMediaImageViewerDetails.h
+++ b/Engine/Plugins/Experimental/MediaViewer/Source/MediaViewer/Private/Widgets/SMediaImageViewerDetails.h
@@ -32,11 +32,11 @@ public:
 	SMediaImageViewerDetails();
 	virtual ~SMediaImageViewerDetails() = default;
 
-	void Construct(const FArguments& InArgs, EMediaImageViewerPosition InPosition, const TSharedPtr<FMediaViewerDelegates>& InDelegates);
+	void Construct(const FArguments& InArgs, EMediaImageViewerPosition InPosition, const TSharedRef<FMediaViewerDelegates>& InDelegates);
 
 protected:
 	EMediaImageViewerPosition Position;
-	FMediaViewerDelegates::FGetImageViewer GetImageViewerDelegate;
+	TSharedPtr<FMediaViewerDelegates> Delegates;
 
 	TSharedPtr<IStructureDetailsView> PanelDetailsView;
 	TSharedPtr<IStructureDetailsView> PaintDetailsView;
diff --git a/Engine/Plugins/Experimental/MediaViewer/Source/MediaViewer/Private/Widgets/SMediaImageViewerOverlay.cpp b/Engine/Plugins/Experimental/MediaViewer/Source/MediaViewer/Private/Widgets/SMediaImageViewerOverlay.cpp
--- a/Engine/Plugins/Experimental/MediaViewer/Source/MediaViewer/Private/Widgets/SMediaImageViewerOverlay.cpp
+++ b/Engine/Plugins/Experimental/MediaViewer/Source/MediaViewer/Private/Widgets/SMediaImageViewerOverlay.cpp
@@ -15,11 +15,11 @@
 #include "Math/Vector2D.h"
 #include "MediaViewer.h"
 #include "MediaViewerCommands.h"
-#include "SAssetDropTarget.h"
 #include "Widgets/Layout/SScissorRectBox.h"
 #include "Widgets/SBoxPanel.h"
 #include "Widgets/SMediaImageViewerStatusBar.h"
 #include "Widgets/SMediaViewer.h"
+#include "Widgets/SMediaViewerDropTarget.h"
 
 #define LOCTEXT_NAMESPACE "SMediaImageViewerOverlay"
 
@@ -40,12 +40,12 @@ void SMediaImageViewerOverlay::PrivateRegisterAttributes(FSlateAttributeDescript
 }
 
 void SMediaImageViewerOverlay::Construct(const FArguments& InArgs, EMediaImageViewerPosition InPosition, 
-	const TSharedPtr<FMediaViewerDelegates>& InDelegates)
+	const TSharedRef<FMediaViewerDelegates>& InDelegates)
 {
 	Position = InPosition;
-	GetImageViewerDelegate = InDelegates->GetImageViewer;
+	Delegates = InDelegates;
 
-	TSharedPtr<FMediaImageViewer> ImageViewer = GetImageViewerDelegate.Execute(Position);
+	TSharedPtr<FMediaImageViewer> ImageViewer = Delegates->GetImageViewer.Execute(Position);
 
 	if (!ImageViewer.IsValid())
 	{
@@ -93,7 +93,7 @@ void SMediaImageViewerOverlay::Construct(const FArguments& InArgs, EMediaImageVi
 
 TSharedPtr<FMediaImageViewer> SMediaImageViewerOverlay::GetImageViewer() const
 {
-	return GetImageViewerDelegate.Execute(Position);
+	return Delegates->GetImageViewer.Execute(Position);
 }
 
 const TSharedRef<FUICommandList>& SMediaImageViewerOverlay::GetCommandList() const
@@ -113,7 +113,7 @@ FIntPoint SMediaImageViewerOverlay::GetImageViewerPixelCoordinates() const
 
 FVector2f SMediaImageViewerOverlay::GetImageViewerPixelCoordinates_Exact() const
 {
-	TSharedPtr<FMediaImageViewer> ImageViewer = GetImageViewerDelegate.Execute(Position);
+	TSharedPtr<FMediaImageViewer> ImageViewer = Delegates->GetImageViewer.Execute(Position);
 
 	if (!ImageViewer.IsValid())
 	{
@@ -137,7 +137,7 @@ FVector2f SMediaImageViewerOverlay::GetImageViewerPixelCoordinates_Exact() const
 
 bool SMediaImageViewerOverlay::IsCursorOverImageViewer() const
 {
-	TSharedPtr<FMediaImageViewer> ImageViewer = GetImageViewerDelegate.Execute(Position);
+	TSharedPtr<FMediaImageViewer> ImageViewer = Delegates->GetImageViewer.Execute(Position);
 
 	if (!ImageViewer.IsValid())
 	{
@@ -201,7 +201,7 @@ FReply SMediaImageViewerOverlay::OnMouseButtonDown(const FGeometry& InMyGeometry
 			return FReply::Handled();
 		}
 
-		if (TSharedPtr<FMediaImageViewer> ImageViewer = GetImageViewerDelegate.Execute(Position))
+		if (TSharedPtr<FMediaImageViewer> ImageViewer = Delegates->GetImageViewer.Execute(Position))
 		{
 			const FInputEventState EventState(nullptr, InMouseEvent.GetEffectingButton(), IE_Pressed);
 
@@ -295,7 +295,7 @@ void SMediaImageViewerOverlay::UpdateMouse(const TOptional<FIntPoint>& InMousePo
 
 	if (InMousePosition.IsSet())
 	{
-		if (TSharedPtr<FMediaImageViewer> ImageViewer = GetImageViewerDelegate.Execute(Position))
+		if (TSharedPtr<FMediaImageViewer> ImageViewer = Delegates->GetImageViewer.Execute(Position))
 		{
 			ImageViewer->OnMouseMove(InMousePosition.GetValue());
 		}
@@ -317,7 +317,7 @@ void SMediaImageViewerOverlay::OnDragButtonUp(FName InKeyName)
 	}
 	else if (bExternalDragging)
 	{
-		if (TSharedPtr<FMediaImageViewer> ImageViewer = GetImageViewerDelegate.Execute(Position))
+		if (TSharedPtr<FMediaImageViewer> ImageViewer = Delegates->GetImageViewer.Execute(Position))
 		{
 			ImageViewer->OnTrackingStopped(Delegates->GetCursorLocation.Execute());
 		}
@@ -338,70 +338,6 @@ void SMediaImageViewerOverlay::UpdateDragPosition()
 	Try_AddOffset(NewOffset - GetOffset());
 }
 
-bool SMediaImageViewerOverlay::OnAssetDraggedOver(TArrayView<FAssetData> InAssetData, FText& OutError) const
-{
-	IMediaViewerModule& Module = IMediaViewerModule::Get();
-
-	for (const FAssetData& AssetData : InAssetData)
-	{
-		if (Module.HasFactoryFor(AssetData))
-		{
-			return true;
-		}
-	}
-
-	OutError = LOCTEXT("NotSupported", "Asset is not supported.");
-
-	return false;
-}
-
-void SMediaImageViewerOverlay::OnAssetsDropped(const FDragDropEvent& InDragDropEvent, TArrayView<FAssetData> InAssetData, 
-	EMediaImageViewerPosition InPosition)
-{
-	IMediaViewerModule& Module = IMediaViewerModule::Get();
-
-	TArray<int32> Indices;
-	Indices.Reserve(static_cast<int32>(EMediaImageViewerPosition::COUNT));
-
-	for (int32 Index = 0; Index < InAssetData.Num(); ++Index)
-	{
-		if (Module.HasFactoryFor(InAssetData[Index]))
-		{
-			Indices.Add(Index);
-
-			if (Indices.Num() == static_cast<int32>(EMediaImageViewerPosition::COUNT))
-			{
-				break;
-			}
-		}
-	}
-
-	if (Indices.IsEmpty())
-	{
-		return;
-	}
-
-	if (Indices.Num() == 1)
-	{
-		Module.SetImage(InPosition, InAssetData[Indices[0]]);
-	}
-	else
-	{
-		Module.SetImage(EMediaImageViewerPosition::First, InAssetData[Indices[0]]);
-		Module.SetImage(EMediaImageViewerPosition::Second, InAssetData[Indices[1]]);
-	}
-}
-
-bool SMediaImageViewerOverlay::OnVerifyDrag(TSharedPtr<FDragDropOperation> InDragDropOp)
-{
-	return true;
-}
-
-FReply SMediaImageViewerOverlay::OnDrop(const FGeometry& InMyGeometry, const FDragDropEvent& InDragDropEvent)
-{
-	return SCompoundWidget::OnDrop(InMyGeometry, InDragDropEvent);
-}
-
 EVisibility SMediaImageViewerOverlay::GetDragDescriptionVisibility() const
 {
 	return FSlateApplication::Get().IsDragDropping()
@@ -467,7 +403,7 @@ void SMediaImageViewerOverlay::BindCommands()
 
 TSharedRef<SWidget> SMediaImageViewerOverlay::CreateOverlay(bool bInCreateSplitDropTarget)
 {
-	TSharedPtr<FMediaImageViewer> ImageViewer = GetImageViewerDelegate.Execute(Position);
+	TSharedPtr<FMediaImageViewer> ImageViewer = Delegates->GetImageViewer.Execute(Position);
 
 	if (!ImageViewer.IsValid())
 	{
@@ -478,86 +414,13 @@ TSharedRef<SWidget> SMediaImageViewerOverlay::CreateOverlay(bool bInCreateSplitD
 
 	if (!bInCreateSplitDropTarget)
 	{
-		return SNew(SAssetDropTarget)
-			.OnAreAssetsAcceptableForDropWithReason(this, &SMediaImageViewerOverlay::OnAssetDraggedOver)
-			.OnAssetsDropped(this, &SMediaImageViewerOverlay::OnAssetsDropped, Position)
-			.bSupportsMultiDrop(true)
+		return SNew(SMediaViewerDropTarget, Delegates.ToSharedRef())
+			.Position(Position)
 			[
-				SNew(SOverlay)
-				+ SOverlay::Slot()
-				[
-					ImageViewerOverlay.IsValid() ? ImageViewerOverlay.ToSharedRef() : SNullWidget::NullWidget
-				]
-				+ SOverlay::Slot()
-				[
-					SNew(SBox)
-					.HAlign(HAlign_Center)
-					.VAlign(VAlign_Center)
-					[
-						SNew(STextBlock)
-						.Visibility(this, &SMediaImageViewerOverlay::GetDragDescriptionVisibility)
-						.Text(LOCTEXT("DropAsset", "Drop supported asset or library item here."))
-					]
-				]
+				ImageViewerOverlay.IsValid() ? ImageViewerOverlay.ToSharedRef() : SNullWidget::NullWidget
 			];
 	}
 
-	TSharedRef<SWidget> FirstDropTarget = SNew(SAssetDropTarget)
-		.OnAreAssetsAcceptableForDropWithReason(this, &SMediaImageViewerOverlay::OnAssetDraggedOver)
-		.OnAssetsDropped(this, &SMediaImageViewerOverlay::OnAssetsDropped, EMediaImageViewerPosition::First)
-		.bSupportsMultiDrop(true)
-		[
-			SNew(SVerticalBox)
-			+ SVerticalBox::Slot()
-			.HAlign(HAlign_Center)
-			.VAlign(VAlign_Bottom)
-			.FillHeight(1.f)
-			.Padding(0.f, 0.f, 0.f, 5.f)
-			[
-				SNew(STextBlock)
-				.Visibility(this, &SMediaImageViewerOverlay::GetDragDescriptionVisibility)
-				.Text(LOCTEXT("ImageA", "Replace Image"))
-			]
-			+ SVerticalBox::Slot()
-			.HAlign(HAlign_Center)
-			.VAlign(VAlign_Top)
-			.FillHeight(1.f)
-			.Padding(0.f, 5.f, 0.f, 0.f)
-			[
-				SNew(STextBlock)
-				.Visibility(this, &SMediaImageViewerOverlay::GetDragDescriptionVisibility)
-				.Text(LOCTEXT("DropAsset", "Drop supported asset or library item here."))
-			]
-		];
-
-	TSharedRef<SWidget> SecondDropTarget = SNew(SAssetDropTarget)
-		.OnAreAssetsAcceptableForDropWithReason(this, &SMediaImageViewerOverlay::OnAssetDraggedOver)
-		.OnAssetsDropped(this, &SMediaImageViewerOverlay::OnAssetsDropped, EMediaImageViewerPosition::Second)
-		.bSupportsMultiDrop(true)
-		[
-			SNew(SVerticalBox)
-			+ SVerticalBox::Slot()
-			.HAlign(HAlign_Center)
-			.VAlign(VAlign_Bottom)
-			.FillHeight(1.f)
-			.Padding(0.f, 0.f, 0.f, 5.f)
-			[
-				SNew(STextBlock)
-				.Visibility(this, &SMediaImageViewerOverlay::GetDragDescriptionVisibility)
-				.Text(LOCTEXT("ImageB", "Add Image"))
-			]
-			+ SVerticalBox::Slot()
-			.HAlign(HAlign_Center)
-			.VAlign(VAlign_Top)
-			.FillHeight(1.f)
-			.Padding(0.f, 5.f, 0.f, 0.f)
-			[
-				SNew(STextBlock)
-				.Visibility(this, &SMediaImageViewerOverlay::GetDragDescriptionVisibility)
-				.Text(LOCTEXT("DropAsset", "Drop supported asset or library item here."))
-			]
-		];
-
 	TSharedPtr<SWidget> SplitDropTarget;
 
 	if (Delegates->GetABOrientation.Execute() == Orient_Horizontal)
@@ -567,14 +430,16 @@ TSharedRef<SWidget> SMediaImageViewerOverlay::CreateOverlay(bool bInCreateSplitD
 			.FillWidth(1.f)
 			.Padding(0.f, 0.f, 2.f, 0.f)
 			[
-				FirstDropTarget
+				SNew(SMediaViewerDropTarget, Delegates.ToSharedRef())
+				.Position(EMediaImageViewerPosition::First)
 			]
 
 			+ SHorizontalBox::Slot()
 			.FillWidth(1.f)
 			.Padding(2.f, 0.f, 0.f, 0.f)
 			[
-				SecondDropTarget
+				SNew(SMediaViewerDropTarget, Delegates.ToSharedRef())
+				.Position(EMediaImageViewerPosition::Second)
 			];
 	}
 	else
@@ -584,14 +449,16 @@ TSharedRef<SWidget> SMediaImageViewerOverlay::CreateOverlay(bool bInCreateSplitD
 			.FillHeight(1.f)
 			.Padding(0.f, 0.f, 0.f, 2.f)
 			[
-				FirstDropTarget
+				SNew(SMediaViewerDropTarget, Delegates.ToSharedRef())
+				.Position(EMediaImageViewerPosition::First)
 			]
 
 			+ SVerticalBox::Slot()
 			.FillHeight(1.f)
 			.Padding(0.f, 2.f, 0.f, 0.f)
 			[
-				SecondDropTarget
+				SNew(SMediaViewerDropTarget, Delegates.ToSharedRef())
+				.Position(EMediaImageViewerPosition::Second)
 			];
 	}
 
@@ -606,7 +473,7 @@ TSharedRef<SWidget> SMediaImageViewerOverlay::CreateOverlay(bool bInCreateSplitD
 		];		
 }
 
-TSharedRef<SWidget> SMediaImageViewerOverlay::CreateStatusBar(const TSharedPtr<FMediaViewerDelegates>& InDelegates)
+TSharedRef<SWidget> SMediaImageViewerOverlay::CreateStatusBar(const TSharedRef<FMediaViewerDelegates>& InDelegates)
 {
 	return SNew(SMediaImageViewerStatusBar, Position, InDelegates);
 }
@@ -665,7 +532,7 @@ void SMediaImageViewerOverlay::ResetTransform()
 
 FVector SMediaImageViewerOverlay::GetOffset() const
 {
-	if (TSharedPtr<FMediaImageViewer> ImageViewer = GetImageViewerDelegate.Execute(Position))
+	if (TSharedPtr<FMediaImageViewer> ImageViewer = Delegates->GetImageViewer.Execute(Position))
 	{
 		return ImageViewer->GetPaintSettings().Offset;
 	}
@@ -692,7 +559,7 @@ void SMediaImageViewerOverlay::AddOffset(const FVector& InOffset)
 
 void SMediaImageViewerOverlay::SetOffset(const FVector& InOffset)
 {
-	if (TSharedPtr<FMediaImageViewer> ImageViewer = GetImageViewerDelegate.Execute(Position))
+	if (TSharedPtr<FMediaImageViewer> ImageViewer = Delegates->GetImageViewer.Execute(Position))
 	{
 		ImageViewer->GetPaintSettings().Offset = InOffset;
 	}
@@ -700,7 +567,7 @@ void SMediaImageViewerOverlay::SetOffset(const FVector& InOffset)
 
 FRotator SMediaImageViewerOverlay::GetRotation() const
 {
-	if (TSharedPtr<FMediaImageViewer> ImageViewer = GetImageViewerDelegate.Execute(Position))
+	if (TSharedPtr<FMediaImageViewer> ImageViewer = Delegates->GetImageViewer.Execute(Position))
 	{
 		return ImageViewer->GetPaintSettings().Rotation;
 	}
@@ -727,7 +594,7 @@ void SMediaImageViewerOverlay::AddRotation(const FRotator& InRotation)
 
 void SMediaImageViewerOverlay::SetRotation(const FRotator& InRotation)
 {
-	if (TSharedPtr<FMediaImageViewer> ImageViewer = GetImageViewerDelegate.Execute(Position))
+	if (TSharedPtr<FMediaImageViewer> ImageViewer = Delegates->GetImageViewer.Execute(Position))
 	{
 		ImageViewer->GetPaintSettings().Rotation = InRotation;
 	}
@@ -735,7 +602,7 @@ void SMediaImageViewerOverlay::SetRotation(const FRotator& InRotation)
 
 float SMediaImageViewerOverlay::GetScale() const
 {
-	if (TSharedPtr<FMediaImageViewer> ImageViewer = GetImageViewerDelegate.Execute(Position))
+	if (TSharedPtr<FMediaImageViewer> ImageViewer = Delegates->GetImageViewer.Execute(Position))
 	{
 		return ImageViewer->GetPaintSettings().Scale;
 	}
@@ -757,7 +624,7 @@ void SMediaImageViewerOverlay::Try_SetScale(float InScale)
 
 void SMediaImageViewerOverlay::SetScale(float InScale)
 {
-	if (TSharedPtr<FMediaImageViewer> ImageViewer = GetImageViewerDelegate.Execute(Position))
+	if (TSharedPtr<FMediaImageViewer> ImageViewer = Delegates->GetImageViewer.Execute(Position))
 	{
 		ImageViewer->GetPaintSettings().Scale = InScale;
 	}
@@ -770,7 +637,7 @@ void SMediaImageViewerOverlay::MultiplyScale(float InMultiple)
 
 void SMediaImageViewerOverlay::ScaleToFit()
 {
-	TSharedPtr<FMediaImageViewer> ImageViewer = GetImageViewerDelegate.Execute(Position);
+	TSharedPtr<FMediaImageViewer> ImageViewer = Delegates->GetImageViewer.Execute(Position);
 
 	if (!ImageViewer.IsValid())
 	{
@@ -836,7 +703,7 @@ void SMediaImageViewerOverlay::CopyTransform()
 
 uint8 SMediaImageViewerOverlay::GetMipLevel() const
 {
-	if (TSharedPtr<FMediaImageViewer> ImageViewer = GetImageViewerDelegate.Execute(Position))
+	if (TSharedPtr<FMediaImageViewer> ImageViewer = Delegates->GetImageViewer.Execute(Position))
 	{
 		return ImageViewer->GetPaintSettings().MipLevel;
 	}
@@ -860,7 +727,7 @@ void SMediaImageViewerOverlay::AdjustMipLevel(int8 InAdjustment)
 
 void SMediaImageViewerOverlay::SetMipLevel(uint8 InMipLevel)
 {
-	if (TSharedPtr<FMediaImageViewer> ImageViewer = GetImageViewerDelegate.Execute(Position))
+	if (TSharedPtr<FMediaImageViewer> ImageViewer = Delegates->GetImageViewer.Execute(Position))
 	{
 		if (InMipLevel < ImageViewer->GetInfo().MipCount)
 		{
@@ -871,7 +738,7 @@ void SMediaImageViewerOverlay::SetMipLevel(uint8 InMipLevel)
 
 void SMediaImageViewerOverlay::CopyColor()
 {
-	TSharedPtr<FMediaImageViewer> ImageViewer = GetImageViewerDelegate.Execute(Position);
+	TSharedPtr<FMediaImageViewer> ImageViewer = Delegates->GetImageViewer.Execute(Position);
 
 	if (!ImageViewer.IsValid() || !ImageViewer->IsValid())
 	{
diff --git a/Engine/Plugins/Experimental/MediaViewer/Source/MediaViewer/Private/Widgets/SMediaImageViewerOverlay.h b/Engine/Plugins/Experimental/MediaViewer/Source/MediaViewer/Private/Widgets/SMediaImageViewerOverlay.h
--- a/Engine/Plugins/Experimental/MediaViewer/Source/MediaViewer/Private/Widgets/SMediaImageViewerOverlay.h
+++ b/Engine/Plugins/Experimental/MediaViewer/Source/MediaViewer/Private/Widgets/SMediaImageViewerOverlay.h
@@ -22,6 +22,7 @@ namespace UE::MediaViewer::Private
 class SMediaViewerLibrary;
 class SMediaViewerLibraryPrivate;
 
+/** Display and control for a viewer. Status bar, mouse and keyboard input, etc. */
 class SMediaImageViewerOverlay : public SCompoundWidget
 {
 	SLATE_DECLARE_WIDGET(SMediaImageViewerOverlay, SCompoundWidget)
@@ -37,23 +38,27 @@ public:
 
 	virtual ~SMediaImageViewerOverlay() override = default;
 
-	void Construct(const FArguments& InArgs, EMediaImageViewerPosition InPosition, const TSharedPtr<FMediaViewerDelegates>& InDelegates);
+	void Construct(const FArguments& InArgs, EMediaImageViewerPosition InPosition, const TSharedRef<FMediaViewerDelegates>& InDelegates);
 
 	TSharedPtr<FMediaImageViewer> GetImageViewer() const;
 
 	const TSharedRef<FUICommandList>& GetCommandList() const;
 
+	/** Gets the pixel coordinate of mouse on the image. May be outside of the image. */
 	FIntPoint GetImageViewerPixelCoordinates() const;
 
+	/** Gets the pixel coordinate of mouse on the image. May be outside of the image. Subpixel accuracy */
 	FVector2f GetImageViewerPixelCoordinates_Exact() const;
 
 	bool IsCursorOverImageViewer() const;
 
+	/** Triggers a mouse update. */
 	void UpdateMouse(const TOptional<FIntPoint>& InMousePosition);
 
 	bool IsOverlayEnabled() const;
 	void ToggleOverlay();
 
+	/** Zooms the view in or out and maintains the cursor in the same relative position. */
 	void MultiplyScaleAroundCursor(float InMultipler);
 
 	//~ Begin SWidget
@@ -67,65 +72,70 @@ public:
 
 protected:
 	EMediaImageViewerPosition Position;
-	FMediaViewerDelegates::FGetImageViewer GetImageViewerDelegate;
+	TSharedPtr<FMediaViewerDelegates> Delegates;
 	TSharedPtr<FMediaViewerLibraryItem> CachedItem;
 	TSharedRef<FUICommandList> CommandList;
+	/** Name of the button which is currently performing a drag operation. */
 	FName DragButtonName;
+	/** Whether the widget is performing an "internal" drag operation. */
 	bool bDragging;
+	/** Whether an image viewer is performing its own drag operation. */
 	bool bExternalDragging;
 	FVector2D DragStartCursor;
 	FVector DragStartOffset;
+	/** For moving the image based on cursor location. */
 	FIntPoint LastCursorPosition;
-	TSharedPtr<FMediaViewerDelegates> Delegates;
 	TSharedPtr<SVerticalBox> Container;
+	/** Toggles the visibility of the overlay. */
 	bool bOverlayEnabled;
 
 	void BindCommands();
 
 	TSharedRef<SWidget> CreateOverlay(bool bInCreateSplitDropTarget);
 
-	TSharedRef<SWidget> CreateStatusBar(const TSharedPtr<FMediaViewerDelegates>& InDelegates);
+	TSharedRef<SWidget> CreateStatusBar(const TSharedRef<FMediaViewerDelegates>& InDelegates);
 
 	void OnDragButtonUp(FName InKeyName);
 	void UpdateDragPosition();
 
-	bool OnAssetDraggedOver(TArrayView<FAssetData> InAssetData, FText& OutError) const;
-	void OnAssetsDropped(const FDragDropEvent& InDragDropEvent, TArrayView<FAssetData> InAssetData, EMediaImageViewerPosition InPosition);
-
-	bool OnVerifyDrag(TSharedPtr<FDragDropOperation> InDragDropOp);
-	FReply OnDrop(const FGeometry& InMyGeometry, const FDragDropEvent& InDragDropEvent);
-
 	EVisibility GetDragDescriptionVisibility() const;
 
+	/** Interact with the image viewer's offset. */
 	FVector GetOffset() const;
-	void Try_AddOffset(FVector InOffset);
+	void Try_AddOffset(FVector InOffset); // Either local or 'send to all' if transform is locked
 	void AddOffset(const FVector& InOffset);
 	void SetOffset(const FVector& InOffset);
 
+	/** Interact with the image viewer's rotation. */
 	FRotator GetRotation() const;
-	void Try_AddRotation(FRotator InRotation);
+	void Try_AddRotation(FRotator InRotation); // Either local or 'send to all' if transform is locked
 	void AddRotation(const FRotator& InRotation);
 	void SetRotation(const FRotator& InRotation);
 
+	/** Interact with the image viewer's scale. */
 	float GetScale() const;
-	void Try_SetScale(float InScale);
+	void Try_SetScale(float InScale); // Either local or 'send to all' if transform is locked
 	void SetScale(float InScale);
 	void MultiplyScale(float InMultiple);
 	void ScaleToFit();
 
+	/** Interact with the image viewer's entire transform. */
 	bool CanResetTransform();
 	void ResetTransform();
-	void Try_ResetTransform();
-	void Try_SetTransform(FVector InOffset, FRotator InRotation, float InScale);
+	void Try_ResetTransform(); // Either local or 'send to all' if transform is locked
+	void Try_SetTransform(FVector InOffset, FRotator InRotation, float InScale); // Either local or 'send to all' if transform is locked
 	bool CanCopyTransform();
 	void CopyTransform();
 
+	/** Interact with the image viewer's mip level. */
 	uint8 GetMipLevel() const;
 	void AdjustMipLevel(int8 InAdjustment);
 	void SetMipLevel(uint8 InMipLevel);
 
+	/** Copies the color under the cursor, if possible. */
 	void CopyColor();
 
+	/** For adding the active image viewer to the library, if it can. */
 	bool CanAddToLibrary() const;
 	void AddToLibrary();
 
diff --git a/Engine/Plugins/Experimental/MediaViewer/Source/MediaViewer/Private/Widgets/SMediaImageViewerStatusBar.cpp b/Engine/Plugins/Experimental/MediaViewer/Source/MediaViewer/Private/Widgets/SMediaImageViewerStatusBar.cpp
--- a/Engine/Plugins/Experimental/MediaViewer/Source/MediaViewer/Private/Widgets/SMediaImageViewerStatusBar.cpp
+++ b/Engine/Plugins/Experimental/MediaViewer/Source/MediaViewer/Private/Widgets/SMediaImageViewerStatusBar.cpp
@@ -28,12 +28,11 @@ void SMediaImageViewerStatusBar::PrivateRegisterAttributes(FSlateAttributeDescri
 }
 
 void SMediaImageViewerStatusBar::Construct(const FArguments& InArgs, EMediaImageViewerPosition InPosition, 
-	const TSharedPtr<FMediaViewerDelegates>& InDelegates)
+	const TSharedRef<FMediaViewerDelegates>& InDelegates)
 {
 	Position = InPosition;
-	GetImageViewerDelegate = InDelegates->GetImageViewer;
-	PixelLocationDelegate = InDelegates->GetPixelCoordinates;
-	check(PixelLocationDelegate.IsBound());
+	Delegates = InDelegates;
+	check(Delegates->GetPixelCoordinates.IsBound());
 
 	ChildSlot
 	[
@@ -43,7 +42,7 @@ void SMediaImageViewerStatusBar::Construct(const FArguments& InArgs, EMediaImage
 
 TSharedRef<SWidget> SMediaImageViewerStatusBar::BuildStatusBar()
 {
-	TSharedPtr<FMediaImageViewer> ImageViewer = GetImageViewerDelegate.Execute(Position);
+	TSharedPtr<FMediaImageViewer> ImageViewer = Delegates->GetImageViewer.Execute(Position);
 
 	if (!ImageViewer.IsValid())
 	{
@@ -111,7 +110,7 @@ TSharedRef<SWidget> SMediaImageViewerStatusBar::BuildStatusBar()
 
 FText SMediaImageViewerStatusBar::GetResolutionLabel() const
 {
-	if (TSharedPtr<FMediaImageViewer> ImageViewer = GetImageViewerDelegate.Execute(Position))
+	if (TSharedPtr<FMediaImageViewer> ImageViewer = Delegates->GetImageViewer.Execute(Position))
 	{
 		return FText::Format(
 			LOCTEXT("Size", "{0} x {1}"),
@@ -125,14 +124,14 @@ FText SMediaImageViewerStatusBar::GetResolutionLabel() const
 
 FText SMediaImageViewerStatusBar::GetColorPickerLabel() const
 {
-	TSharedPtr<FMediaImageViewer> ImageViewer = GetImageViewerDelegate.Execute(Position);
+	TSharedPtr<FMediaImageViewer> ImageViewer = Delegates->GetImageViewer.Execute(Position);
 
 	if (!ImageViewer.IsValid() || !ImageViewer->IsValid())
 	{
 		return INVTEXT("-");
 	}
 
-	const FIntPoint PixelCoordinates = PixelLocationDelegate.Execute(Position);
+	const FIntPoint PixelCoordinates = Delegates->GetPixelCoordinates.Execute(Position);
 
 	if (PixelCoordinates.X < 0 || PixelCoordinates.Y < 0)
 	{
diff --git a/Engine/Plugins/Experimental/MediaViewer/Source/MediaViewer/Private/Widgets/SMediaImageViewerStatusBar.h b/Engine/Plugins/Experimental/MediaViewer/Source/MediaViewer/Private/Widgets/SMediaImageViewerStatusBar.h
--- a/Engine/Plugins/Experimental/MediaViewer/Source/MediaViewer/Private/Widgets/SMediaImageViewerStatusBar.h
+++ b/Engine/Plugins/Experimental/MediaViewer/Source/MediaViewer/Private/Widgets/SMediaImageViewerStatusBar.h
@@ -30,12 +30,11 @@ public:
 
 	virtual ~SMediaImageViewerStatusBar() override = default;
 
-	void Construct(const FArguments& InArgs, EMediaImageViewerPosition InPosition, const TSharedPtr<FMediaViewerDelegates>& InDelegates);
+	void Construct(const FArguments& InArgs, EMediaImageViewerPosition InPosition, const TSharedRef<FMediaViewerDelegates>& InDelegates);
 
 protected:
 	EMediaImageViewerPosition Position;
-	FMediaViewerDelegates::FGetImageViewer GetImageViewerDelegate;
-	FMediaViewerDelegates::FGetLocationForPosition PixelLocationDelegate;
+	TSharedPtr<FMediaViewerDelegates> Delegates;
 
 	TSharedRef<SWidget> BuildStatusBar();
 
diff --git a/Engine/Plugins/Experimental/MediaViewer/Source/MediaViewer/Private/Widgets/SMediaViewer.cpp b/Engine/Plugins/Experimental/MediaViewer/Source/MediaViewer/Private/Widgets/SMediaViewer.cpp
--- a/Engine/Plugins/Experimental/MediaViewer/Source/MediaViewer/Private/Widgets/SMediaViewer.cpp
+++ b/Engine/Plugins/Experimental/MediaViewer/Source/MediaViewer/Private/Widgets/SMediaViewer.cpp
@@ -4,7 +4,6 @@
 
 #include "Brushes/SlateColorBrush.h"
 #include "Brushes/SlateImageBrush.h"
-#include "Editor/EditorWidgets/Public/SAssetDropTarget.h"
 #include "Engine/Texture2D.h"
 #include "Framework/Application/SlateApplication.h"
 #include "ImageViewer/MediaImageViewer.h"
@@ -17,6 +16,7 @@
 #include "Widgets/Layout/SSplitter.h"
 #include "Widgets/MediaViewerDelegates.h"
 #include "Widgets/SMediaImageViewerOverlay.h"
+#include "Widgets/SMediaViewerDropTarget.h"
 #include "Widgets/SMediaViewerLibrary.h"
 #include "Widgets/SMediaViewerLibraryPrivate.h"
 #include "Widgets/SMediaViewerToolbar.h"
@@ -58,7 +58,7 @@ void SMediaViewer::Construct(const FArguments& InArgs, const FMediaViewerArgs& I
 	BindCommands();
 	CreateDelegates();
 
-	Library = SNew(SMediaViewerLibraryPrivate, Delegates);
+	Library = SNew(SMediaViewerLibraryPrivate, Delegates.ToSharedRef());
 
 	Layout = SNew(SVerticalBox);
 
@@ -151,7 +151,7 @@ void SMediaViewer::SetImageViewer(EMediaImageViewerPosition InPosition, const TS
 		Library->OnImageViewerOpened(InImageViewer);
 	}
 
-	CreateView();
+	InvalidateView();
 }
 
 void SMediaViewer::BindCommands()
@@ -494,6 +494,7 @@ void SMediaViewer::CreateDelegates()
 
 	// Per panel delegates
 	Delegates->GetImageViewer.BindSP(this, &SMediaViewer::GetImageViewer);
+	Delegates->SetImageViewer.BindSP(this, &SMediaViewer::SetImageViewer);
 	Delegates->GetPixelCoordinates.BindSP(this, &SMediaViewer::GetPixelCoordinates);
 	Delegates->IsOverImage.BindSP(this, &SMediaViewer::IsOverImage);
 	Delegates->CopyTransformToAll.BindSP(this, &SMediaViewer::CopyTransformToAll);
@@ -502,13 +503,13 @@ void SMediaViewer::CreateDelegates()
 
 TSharedRef<SMediaImageViewerOverlay> SMediaViewer::CreateOverlay(EMediaImageViewerPosition InPosition, bool bInCreateSplitDropTarget)
 {
-	return  SNew(SMediaImageViewerOverlay, InPosition, Delegates)
+	return  SNew(SMediaImageViewerOverlay, InPosition, Delegates.ToSharedRef())
 		.bCreateSplitDropTarget(bInCreateSplitDropTarget);
 }
 
 TSharedRef<SWidget> SMediaViewer::CreateToolbar()
 {
-	return SNew(SMediaViewerToolbar, Delegates);
+	return SNew(SMediaViewerToolbar, Delegates.ToSharedRef());
 }
 
 void SMediaViewer::CreateView()
@@ -556,10 +557,7 @@ void SMediaViewer::CreateEmptyView()
 {
 	// Attach a null widget
 	ContentSlot->AttachWidget(
-		SNew(SAssetDropTarget)
-		.OnAreAssetsAcceptableForDrop(this, &SMediaViewer::OnAssetDraggedOver)
-		.OnAssetsDropped(this, &SMediaViewer::OnAssetsDropped)
-		.bSupportsMultiDrop(true)
+		SNew(SMediaViewerDropTarget, Delegates.ToSharedRef())
 	);
 }
 
@@ -681,57 +679,6 @@ bool SMediaViewer::IsImageViewerNull(EMediaImageViewerPosition InPosition) const
 	return ImageViewers[static_cast<int32>(InPosition)] == FNullImageViewer::GetNullImageViewer();
 }
 
-bool SMediaViewer::OnAssetDraggedOver(TArrayView<FAssetData> InAssetData) const
-{
-	IMediaViewerModule& Module = IMediaViewerModule::Get();
-
-	for (const FAssetData& AssetData : InAssetData)
-	{
-		if (Module.HasFactoryFor(AssetData))
-		{
-			return true;
-		}
-	}
-
-	return false;
-}
-
-void SMediaViewer::OnAssetsDropped(const FDragDropEvent& InDragDropEvent, TArrayView<FAssetData> InAssetData)
-{
-	IMediaViewerModule& Module = IMediaViewerModule::Get();
-
-	TArray<int32> Indices;
-	Indices.Reserve(static_cast<int32>(EMediaImageViewerPosition::COUNT));
-
-	for (int32 Index = 0; Index < InAssetData.Num(); ++Index)
-	{
-		if (Module.HasFactoryFor(InAssetData[Index]))
-		{
-			Indices.Add(Index);
-
-			if (Indices.Num() == static_cast<int32>(EMediaImageViewerPosition::COUNT))
-			{
-				break;
-			}
-		}
-	}
-
-	if (Indices.IsEmpty())
-	{
-		return;
-	}
-
-	if (Indices.Num() == 1)
-	{
-		Module.SetImage(EMediaImageViewerPosition::First, InAssetData[Indices[0]]);
-	}
-	else
-	{
-		Module.SetImage(EMediaImageViewerPosition::First, InAssetData[Indices[0]]);
-		Module.SetImage(EMediaImageViewerPosition::Second, InAssetData[Indices[1]]);
-	}
-}
-
 void SMediaViewer::AddOffsetToAll(const FVector& InOffset)
 {
 	for (const TSharedPtr<FMediaImageViewer>& ImageViewer : ImageViewers)
diff --git a/Engine/Plugins/Experimental/MediaViewer/Source/MediaViewer/Private/Widgets/SMediaViewer.h b/Engine/Plugins/Experimental/MediaViewer/Source/MediaViewer/Private/Widgets/SMediaViewer.h
--- a/Engine/Plugins/Experimental/MediaViewer/Source/MediaViewer/Private/Widgets/SMediaViewer.h
+++ b/Engine/Plugins/Experimental/MediaViewer/Source/MediaViewer/Private/Widgets/SMediaViewer.h
@@ -171,12 +171,6 @@ protected:
 	/** Returns trus if the image viewer is null or an FNullImageViewer. */
 	bool IsImageViewerNull(EMediaImageViewerPosition InPosition) const;
 
-	/** Activated only in "empty" view to easily add new assets when no overlays are present. */
-	bool OnAssetDraggedOver(TArrayView<FAssetData> InAssetData) const;
-
-	/** Activated only in "empty" view to easily add new assets when no overlays are present. */
-	void OnAssetsDropped(const FDragDropEvent& InDragDropEvent, TArrayView<FAssetData> InAssetData);
-
 	/** Adds an offset to all image viewers. */
 	void AddOffsetToAll(const FVector& InOffset);
 
diff --git a/Engine/Plugins/Experimental/MediaViewer/Source/MediaViewer/Private/Widgets/SMediaViewerDropTarget.cpp b/Engine/Plugins/Experimental/MediaViewer/Source/MediaViewer/Private/Widgets/SMediaViewerDropTarget.cpp
new file mode 100644
--- /dev/null
+++ b/Engine/Plugins/Experimental/MediaViewer/Source/MediaViewer/Private/Widgets/SMediaViewerDropTarget.cpp
@@ -0,0 +1,297 @@
+// Copyright Epic Games, Inc. All Rights Reserved.
+
+#include "Widgets/SMediaViewerDropTarget.h"
+
+#include "AssetSelection.h"
+#include "DragAndDrop/DecoratedDragDropOp.h"
+#include "Framework/Application/SlateApplication.h"
+#include "IMediaViewerModule.h"
+#include "Library/MediaViewerLibrary.h"
+#include "Library/MediaViewerLibraryItem.h"
+#include "Styling/AppStyle.h"
+#include "Styling/StyleColors.h"
+#include "Widgets/MediaViewerDelegates.h"
+#include "Widgets/MediaViewerLibraryItemDragDropOperation.h"
+#include "Widgets/SBoxPanel.h"
+#include "Widgets/SOverlay.h"
+
+#define LOCTEXT_NAMESPACE "SMediaViewerDropTarget"
+
+namespace UE::MediaViewer::Private
+{
+
+void SMediaViewerDropTarget::PrivateRegisterAttributes(FSlateAttributeDescriptor::FInitializer&)
+{
+}
+
+void SMediaViewerDropTarget::Construct(const FArguments& InArgs, const TSharedRef<FMediaViewerDelegates>& InDelegates)
+{
+	Delegates = InDelegates;
+	Position = InArgs._Position;
+
+	const EMediaImageViewerActivePosition ActiveViews = Delegates->GetActiveViews.Execute();
+	
+	const FText DropTargetMessageLower = LOCTEXT("DropTargetMessage", "Drop supported asset or library item here.");
+	FText DropTargetMessageUpper = FText::GetEmpty();
+
+	switch (ActiveViews)
+	{
+		default:
+		case EMediaImageViewerActivePosition::None:
+			break;
+
+		case EMediaImageViewerActivePosition::First:
+		case EMediaImageViewerActivePosition::Second:
+			if (Position == EMediaImageViewerPosition::First)
+			{
+				DropTargetMessageUpper = LOCTEXT("DropTargetMessageReplaceImage", "Replace Image");
+			}
+			else if (Position == EMediaImageViewerPosition::Second)
+			{
+				DropTargetMessageUpper = LOCTEXT("DropTargetMessageAddSecond", "Add Second Image");
+			}
+			break;
+
+		case EMediaImageViewerActivePosition::Both:
+			if (Position == EMediaImageViewerPosition::First || Position == EMediaImageViewerPosition::Second)
+			{
+				DropTargetMessageUpper = LOCTEXT("DropTargetMessageReplaceImage", "Replace Image");
+			}
+			break;
+	}
+
+	TSharedRef<SVerticalBox> MessageBox = SNew(SVerticalBox);
+
+	if (DropTargetMessageUpper.IsEmpty())
+	{
+		MessageBox->AddSlot()
+			.FillHeight(1.f)
+			.HAlign(HAlign_Center)
+			.VAlign(VAlign_Center)
+			.Padding(5.f)
+			[
+				SNew(STextBlock)
+				.ColorAndOpacity(FStyleColors::Foreground.GetSpecifiedColor())
+				.ShadowColorAndOpacity(FStyleColors::Panel.GetSpecifiedColor())
+				.ShadowOffset(FVector2D(1.f, 1.f))
+				.Visibility(this, &SMediaViewerDropTarget::GetDragDescriptionVisibility)
+				.Text(DropTargetMessageLower)
+			];
+	}
+	else
+	{
+		MessageBox->AddSlot()
+			.FillHeight(1.f)
+			.HAlign(HAlign_Center)
+			.VAlign(VAlign_Bottom)
+			.Padding(5.f)
+			[
+				SNew(STextBlock)
+				.ColorAndOpacity(FStyleColors::Foreground.GetSpecifiedColor())
+				.ShadowColorAndOpacity(FStyleColors::Panel.GetSpecifiedColor())
+				.ShadowOffset(FVector2D(1.f, 1.f))
+				.Visibility(this, &SMediaViewerDropTarget::GetDragDescriptionVisibility)
+				.Text(DropTargetMessageUpper)
+			];
+
+		MessageBox->AddSlot()
+			.FillHeight(1.f)
+			.HAlign(HAlign_Center)
+			.VAlign(VAlign_Top)
+			.Padding(5.f)
+			[
+				SNew(STextBlock)
+				.ColorAndOpacity(FStyleColors::Foreground.GetSpecifiedColor())
+				.ShadowColorAndOpacity(FStyleColors::Panel.GetSpecifiedColor())
+				.ShadowOffset(FVector2D(1.f, 1.f))
+				.Visibility(this, &SMediaViewerDropTarget::GetDragDescriptionVisibility)
+				.Text(DropTargetMessageLower)
+			];
+	}
+
+	SDropTarget::Construct(
+		SDropTarget::FArguments()
+		.OnDropped(this, &SMediaViewerDropTarget::OnDropped)
+		[
+			SNew(SOverlay)
+			+ SOverlay::Slot()
+			[
+				InArgs._Content.Widget
+			]
+			+ SOverlay::Slot()
+			[
+				MessageBox
+			]
+		]);
+}
+
+void SMediaViewerDropTarget::OnDragLeave(const FDragDropEvent& InDragDropEvent)
+{
+	TSharedPtr<FDragDropOperation> Operation = InDragDropEvent.GetOperation();
+
+	if (Operation.IsValid() && Operation->IsOfType<FDecoratedDragDropOp>())
+	{
+		TSharedPtr<FDecoratedDragDropOp> DragDropOp = StaticCastSharedPtr<FDecoratedDragDropOp>(Operation);
+		DragDropOp->ResetToDefaultToolTip();
+	}
+}
+
+TArray<FAssetData> SMediaViewerDropTarget::GetAssetsWithImageViewer(TConstArrayView<FAssetData> InAssets)
+{
+	IMediaViewerModule& MediaViewerModule = IMediaViewerModule::Get();
+
+	TArray<FAssetData> ValidAssets;
+
+	for (const FAssetData& AssetData : InAssets)
+	{
+		if (MediaViewerModule.HasFactoryFor(AssetData))
+		{
+			ValidAssets.Add(AssetData);
+		}
+	}
+
+	return ValidAssets;
+}
+
+EVisibility SMediaViewerDropTarget::GetDragDescriptionVisibility() const
+{
+	return FSlateApplication::Get().IsDragDropping()
+		? EVisibility::HitTestInvisible
+		: EVisibility::Collapsed;
+}
+
+FReply SMediaViewerDropTarget::OnDropped(const FGeometry& InGeometry, const FDragDropEvent& InDragDropEvent)
+{
+	TSharedPtr<FDragDropOperation> DragDropOp = InDragDropEvent.GetOperation();
+
+	if (!DragDropOp.IsValid())
+	{
+		return FReply::Handled();
+	}
+
+	if (DragDropOp->IsOfType<FMediaViewerLibraryItemDragDropOperation>())
+	{
+		HandleDroppedMediaViewerOp(*StaticCastSharedPtr<FMediaViewerLibraryItemDragDropOperation>(DragDropOp));
+		return FReply::Handled();
+	}
+
+	TArray<FAssetData> DroppedAssets = AssetUtil::ExtractAssetDataFromDrag(DragDropOp);
+
+	if (!DroppedAssets.IsEmpty())
+	{
+		HandleDroppedAssets(DroppedAssets);
+	}
+
+	return FReply::Handled();
+}
+
+void SMediaViewerDropTarget::HandleDroppedMediaViewerOp(const FMediaViewerLibraryItemDragDropOperation& InMediaViewerOp)
+{
+	TSharedRef<FMediaViewerLibrary> Library = Delegates->GetLibrary.Execute();
+
+	if (TSharedPtr<FMediaViewerLibraryItem> LibraryItem = Library->GetItem(InMediaViewerOp.GetGroupItem().ItemId))
+	{
+		if (TSharedPtr<FMediaImageViewer> ImageViewer = LibraryItem->CreateImageViewer())
+		{
+			Delegates->SetImageViewer.Execute(Position, ImageViewer.ToSharedRef());
+		}
+	}
+}
+
+void SMediaViewerDropTarget::HandleDroppedAssets(const TArrayView<FAssetData> InDroppedAssets)
+{
+	IMediaViewerModule& MediaViewerModule = IMediaViewerModule::Get();
+
+	const TArray<FAssetData> ValidAssets = GetAssetsWithImageViewer(InDroppedAssets);
+	TArray<TSharedRef<FMediaImageViewer>> ImageViewers;
+
+	for (const FAssetData& AssetData : ValidAssets)
+	{
+		if (TSharedPtr<FMediaViewerLibraryItem> LibraryItem = MediaViewerModule.CreateLibraryItem(AssetData))
+		{
+			if (TSharedPtr<FMediaImageViewer> ImageViewer = LibraryItem->CreateImageViewer())
+			{
+				ImageViewers.Add(ImageViewer.ToSharedRef());
+
+				if (ImageViewers.Num() == static_cast<int32>(EMediaImageViewerPosition::COUNT))
+				{
+					break;
+				}
+			}
+		}
+	}
+
+	if (ImageViewers.Num() > 1)
+	{
+		for (int32 Index = 0; Index < static_cast<int32>(EMediaImageViewerPosition::COUNT); ++Index)
+		{
+			Delegates->SetImageViewer.Execute(static_cast<EMediaImageViewerPosition>(Index), ImageViewers[Index]);
+		}
+	}
+	else if (ImageViewers.Num() == 1)
+	{
+		Delegates->SetImageViewer.Execute(Position, ImageViewers[0]);
+	}
+}
+
+bool SMediaViewerDropTarget::OnAllowDrop(TSharedPtr<FDragDropOperation> InDragDropOperation) const
+{
+	if (!InDragDropOperation.IsValid())
+	{
+		return false;
+	}
+
+	if (InDragDropOperation->IsOfType<FMediaViewerLibraryItemDragDropOperation>())
+	{
+		TSharedPtr<FMediaViewerLibraryItemDragDropOperation> LibraryItemDragDropOp = StaticCastSharedPtr<FMediaViewerLibraryItemDragDropOperation>(InDragDropOperation);
+
+		if (!Delegates->GetLibrary.Execute()->GetItem(LibraryItemDragDropOp->GetGroupItem().ItemId).IsValid())
+		{
+			if (IsDragOver())
+			{
+				LibraryItemDragDropOp->SetToolTip(LOCTEXT("InvalidItem", "Invalid Library Item"), FAppStyle::GetBrush(TEXT("Graph.ConnectorFeedback.Error")));
+			}
+
+			return false;
+		}
+
+		return true;
+	}
+
+	const TArray<FAssetData> DroppedAssets = AssetUtil::ExtractAssetDataFromDrag(InDragDropOperation);
+	const TArray<FAssetData> ValidAssets = GetAssetsWithImageViewer(DroppedAssets);
+
+	if (ValidAssets.IsEmpty())
+	{
+		if (InDragDropOperation->IsOfType<FDecoratedDragDropOp>())
+		{
+			TSharedPtr<FDecoratedDragDropOp> DecoratedDragDropOp = StaticCastSharedPtr<FDecoratedDragDropOp>(InDragDropOperation);
+			DecoratedDragDropOp->SetToolTip(LOCTEXT("NotSupported", "Not Supported"), FAppStyle::GetBrush(TEXT("Graph.ConnectorFeedback.Error")));
+		}
+
+		return false;
+	}
+
+	return true;
+}
+
+bool SMediaViewerDropTarget::OnIsRecognized(TSharedPtr<FDragDropOperation> InDragDropOperation) const
+{
+	if (!InDragDropOperation.IsValid())
+	{
+		return false;
+	}
+
+	if (InDragDropOperation->IsOfType<FMediaViewerLibraryItemDragDropOperation>())
+	{
+		return true;
+	}
+
+	const TArray<FAssetData> DroppedAssets = AssetUtil::ExtractAssetDataFromDrag(InDragDropOperation);
+
+	return !DroppedAssets.IsEmpty();
+}
+
+}
+
+#undef LOCTEXT_NAMESPACE
diff --git a/Engine/Plugins/Experimental/MediaViewer/Source/MediaViewer/Private/Widgets/SMediaViewerDropTarget.h b/Engine/Plugins/Experimental/MediaViewer/Source/MediaViewer/Private/Widgets/SMediaViewerDropTarget.h
new file mode 100644
--- /dev/null
+++ b/Engine/Plugins/Experimental/MediaViewer/Source/MediaViewer/Private/Widgets/SMediaViewerDropTarget.h
@@ -0,0 +1,53 @@
+// Copyright Epic Games, Inc. All Rights Reserved.
+
+#pragma once
+
+#include "SDropTarget.h"
+
+#include "SMediaViewer.h"
+
+namespace UE::MediaViewer::Private
+{
+
+class FMediaViewerLibraryItemDragDropOperation;
+struct FMediaViewerDelegates;
+
+class SMediaViewerDropTarget : public SDropTarget
+{
+	SLATE_DECLARE_WIDGET(SMediaViewerDropTarget, SDropTarget)
+
+	SLATE_BEGIN_ARGS(SMediaViewerDropTarget)
+		: _Position(EMediaImageViewerPosition::First)
+		{}
+		SLATE_DEFAULT_SLOT(FArguments, Content)
+		SLATE_ARGUMENT(EMediaImageViewerPosition, Position)
+	SLATE_END_ARGS()
+
+public:
+	void Construct(const FArguments& InArgs, const TSharedRef<FMediaViewerDelegates>& InDelegates);
+
+	//~ Begin SWidget
+	virtual void OnDragLeave(const FDragDropEvent& InDragDropEvent) override;
+	//~ End SWidget
+
+protected:
+	static TArray<FAssetData> GetAssetsWithImageViewer(TConstArrayView<FAssetData> InAssets);
+
+	TSharedPtr<FMediaViewerDelegates> Delegates;
+	EMediaImageViewerPosition Position;
+
+	EVisibility GetDragDescriptionVisibility() const;
+
+	FReply OnDropped(const FGeometry& InGeometry, const FDragDropEvent& InDragDropEvent);
+
+	void HandleDroppedMediaViewerOp(const FMediaViewerLibraryItemDragDropOperation& InMediaViewerOp);
+
+	void HandleDroppedAssets(const TArrayView<FAssetData> InDroppedAssets);
+
+	//~ Begin SDropTarget
+	virtual bool OnAllowDrop(TSharedPtr<FDragDropOperation> InDragDropOperation) const override;
+	virtual bool OnIsRecognized(TSharedPtr<FDragDropOperation> InDragDropOperation) const override;
+	//~ End SDropTarget
+};
+
+} // UE::MediaViewer::Private
diff --git a/Engine/Plugins/Experimental/MediaViewer/Source/MediaViewer/Private/Widgets/SMediaViewerLibrary.cpp b/Engine/Plugins/Experimental/MediaViewer/Source/MediaViewer/Private/Widgets/SMediaViewerLibrary.cpp
--- a/Engine/Plugins/Experimental/MediaViewer/Source/MediaViewer/Private/Widgets/SMediaViewerLibrary.cpp
+++ b/Engine/Plugins/Experimental/MediaViewer/Source/MediaViewer/Private/Widgets/SMediaViewerLibrary.cpp
@@ -25,8 +25,10 @@ void SMediaViewerLibrary::PrivateRegisterAttributes(FSlateAttributeDescriptor::F
 {
 }
 
-void SMediaViewerLibrary::Construct(const FArguments& InArgs, const IMediaViewerLibraryWidget::FArgs& InMediaViewerLibraryArgs)
+void SMediaViewerLibrary::Construct(const FArguments& InArgs, const IMediaViewerLibraryWidget::FArgs& InMediaViewerLibraryArgs,
+	const TSharedRef<FMediaViewerDelegates>& InDelegates)
 {
+	Delegates = InDelegates;
 	Library = MakeShared<FMediaViewerLibrary>(InMediaViewerLibraryArgs.LibraryArgs);
 	Library->GetOnChanged().BindSP(this, &SMediaViewerLibrary::OnLibraryChanged);
 
@@ -102,7 +104,7 @@ TSharedRef<class ITableRow> SMediaViewerLibrary::OnGenerateItemRow(IMediaViewerL
 {
 	if (InEntry.ItemId.IsValid())
 	{
-		return SNew(SMediaViewerLibraryItem, InOwningTable, Library.ToSharedRef(), InEntry);
+		return SNew(SMediaViewerLibraryItem, InOwningTable, Library.ToSharedRef(), InEntry, Delegates.ToSharedRef());
 	}
 
 	return SNew(SMediaViewerLibraryGroup, InOwningTable, Library.ToSharedRef(), InEntry.GroupId);
diff --git a/Engine/Plugins/Experimental/MediaViewer/Source/MediaViewer/Private/Widgets/SMediaViewerLibrary.h b/Engine/Plugins/Experimental/MediaViewer/Source/MediaViewer/Private/Widgets/SMediaViewerLibrary.h
--- a/Engine/Plugins/Experimental/MediaViewer/Source/MediaViewer/Private/Widgets/SMediaViewerLibrary.h
+++ b/Engine/Plugins/Experimental/MediaViewer/Source/MediaViewer/Private/Widgets/SMediaViewerLibrary.h
@@ -21,6 +21,7 @@ namespace UE::MediaViewer::Private
 {
 
 class FMediaViewerLibrary;
+struct FMediaViewerDelegates;
 
 /**
  * Implementation of @see IMediaViewerLibraryWidget.
@@ -38,7 +39,7 @@ public:
 
 	virtual ~SMediaViewerLibrary() override = default;
 
-	void Construct(const FArguments& InArgs, const FArgs& InMediaViewerLibraryArgs);
+	void Construct(const FArguments& InArgs, const FArgs& InMediaViewerLibraryArgs, const TSharedRef<FMediaViewerDelegates>& InDelegates);
 
 	virtual TSharedRef<FMediaViewerLibrary> GetLibraryImpl() const;
 
@@ -54,6 +55,7 @@ public:
 	//~ End SWidget
 
 protected:
+	TSharedPtr<FMediaViewerDelegates> Delegates;
 	TSharedPtr<FMediaViewerLibrary> Library;
 	TArray<IMediaViewerLibrary::FGroupItem> Groups;
 	TSharedPtr<STreeView<IMediaViewerLibrary::FGroupItem>> TreeView;
diff --git a/Engine/Plugins/Experimental/MediaViewer/Source/MediaViewer/Private/Widgets/SMediaViewerLibraryItem.cpp b/Engine/Plugins/Experimental/MediaViewer/Source/MediaViewer/Private/Widgets/SMediaViewerLibraryItem.cpp
--- a/Engine/Plugins/Experimental/MediaViewer/Source/MediaViewer/Private/Widgets/SMediaViewerLibraryItem.cpp
+++ b/Engine/Plugins/Experimental/MediaViewer/Source/MediaViewer/Private/Widgets/SMediaViewerLibraryItem.cpp
@@ -16,6 +16,7 @@
 #include "Widgets/Colors/SColorBlock.h"
 #include "Widgets/Images/SImage.h"
 #include "Widgets/Input/SButton.h"
+#include "Widgets/MediaViewerDelegates.h"
 #include "Widgets/MediaViewerLibraryItemDragDropOperation.h"
 #include "Widgets/SBoxPanel.h"
 #include "Widgets/Text/STextBlock.h"
@@ -34,13 +35,21 @@ void SMediaViewerLibraryItem::PrivateRegisterAttributes(FSlateAttributeDescripto
 }
 
 void SMediaViewerLibraryItem::Construct(const FArguments& InArgs, const TSharedRef<STableViewBase>& InOwningTable, 
-	const TSharedRef<FMediaViewerLibrary>& InLibrary, const IMediaViewerLibrary::FGroupItem& InGroupItem)
+	const TSharedRef<FMediaViewerLibrary>& InLibrary, const IMediaViewerLibrary::FGroupItem& InGroupItem,
+	const TSharedRef<FMediaViewerDelegates>& InDelegates)
 {
 	LibraryWeak = InLibrary;
 	GroupItem = InGroupItem;
+	Delegates = InDelegates;
+
+	Item = InLibrary->GetItem(InGroupItem.ItemId);
+
+	if (!Item.IsValid())
+	{
+		return;
+	}
 
 	TSharedPtr<FMediaViewerLibraryGroup> Group = InLibrary->GetGroup(GroupItem.GroupId);
-	TSharedPtr<FMediaViewerLibraryItem> Item = InLibrary->GetItem(InGroupItem.ItemId);
 
 	SetCursor(EMouseCursor::Hand);
 
@@ -194,8 +203,6 @@ FReply SMediaViewerLibraryItem::OnUseButtonClicked(EMediaImageViewerPosition InP
 		return FReply::Handled();
 	}
 
-	TSharedPtr<FMediaViewerLibraryItem> Item = Library->GetItem(GroupItem.ItemId);
-
 	if (!Item.IsValid())
 	{
 		return FReply::Handled();
@@ -208,8 +215,7 @@ FReply SMediaViewerLibraryItem::OnUseButtonClicked(EMediaImageViewerPosition InP
 		return FReply::Handled();
 	}
 
-	IMediaViewerModule& Module = IMediaViewerModule::Get();
-	Module.SetImage(InPosition, ImageViewer.ToSharedRef());
+	Delegates->SetImageViewer.Execute(InPosition, ImageViewer.ToSharedRef());
 
 	return FReply::Handled();
 }
@@ -233,13 +239,6 @@ FReply SMediaViewerLibraryItem::OnRemoveButtonClicked()
 		return FReply::Handled();
 	}
 
-	TSharedPtr<FMediaViewerLibraryItem> Item = Library->GetItem(GroupItem.ItemId);
-
-	if (!Item.IsValid())
-	{
-		return FReply::Handled();
-	}
-
 	Library->RemoveItemFromGroup(GroupItem);
 
 	return FReply::Handled();
diff --git a/Engine/Plugins/Experimental/MediaViewer/Source/MediaViewer/Private/Widgets/SMediaViewerLibraryItem.h b/Engine/Plugins/Experimental/MediaViewer/Source/MediaViewer/Private/Widgets/SMediaViewerLibraryItem.h
--- a/Engine/Plugins/Experimental/MediaViewer/Source/MediaViewer/Private/Widgets/SMediaViewerLibraryItem.h
+++ b/Engine/Plugins/Experimental/MediaViewer/Source/MediaViewer/Private/Widgets/SMediaViewerLibraryItem.h
@@ -24,6 +24,7 @@ namespace UE::MediaViewer::Private
 {
 
 class FMediaViewerLibrary;
+struct FMediaViewerDelegates;
 	
 class SMediaViewerLibraryItem : public STableRow<IMediaViewerLibrary::FGroupItem>
 {
@@ -41,7 +42,8 @@ public:
 	virtual ~SMediaViewerLibraryItem() override = default;
 
 	void Construct(const FArguments& InArgs, const TSharedRef<STableViewBase>& InOwningTable,
-		const TSharedRef<FMediaViewerLibrary>& InLibrary, const IMediaViewerLibrary::FGroupItem& InGroupItem);
+		const TSharedRef<FMediaViewerLibrary>& InLibrary, const IMediaViewerLibrary::FGroupItem& InGroupItem, 
+		const TSharedRef<FMediaViewerDelegates>& InDelegates);
 
 	//~ Begin STableRow
 	virtual void OnDragEnter(const FGeometry& InMyGeometry, const FDragDropEvent& InEvent) override;
@@ -51,6 +53,8 @@ public:
 protected:
 	TWeakPtr<FMediaViewerLibrary> LibraryWeak;
 	IMediaViewerLibrary::FGroupItem GroupItem;
+	TSharedPtr<FMediaViewerDelegates> Delegates;
+	TSharedPtr<FMediaViewerLibraryItem> Item;
 	TSharedPtr<FSlateBrush> ThumbnailBrush;
 
 	FReply OnUseButtonClicked(EMediaImageViewerPosition InPosition);
diff --git a/Engine/Plugins/Experimental/MediaViewer/Source/MediaViewer/Private/Widgets/SMediaViewerLibraryPrivate.cpp b/Engine/Plugins/Experimental/MediaViewer/Source/MediaViewer/Private/Widgets/SMediaViewerLibraryPrivate.cpp
--- a/Engine/Plugins/Experimental/MediaViewer/Source/MediaViewer/Private/Widgets/SMediaViewerLibraryPrivate.cpp
+++ b/Engine/Plugins/Experimental/MediaViewer/Source/MediaViewer/Private/Widgets/SMediaViewerLibraryPrivate.cpp
@@ -3,7 +3,6 @@
 #include "Widgets/SMediaViewerLibraryPrivate.h"
 
 #include "ImageViewers/ColorImageViewer.h"
-#include "IMediaViewerModule.h"
 #include "LevelEditor.h"
 #include "Library/LevelEditorViewportGroup.h"
 #include "Library/MediaTextureGroup.h"
@@ -37,13 +36,13 @@ void SMediaViewerLibraryPrivate::PrivateRegisterAttributes(FSlateAttributeDescri
 {
 }
 
-void SMediaViewerLibraryPrivate::Construct(const FArguments& InArgs, const TSharedPtr<FMediaViewerDelegates>& InDelegates)
+void SMediaViewerLibraryPrivate::Construct(const FArguments& InArgs, const TSharedRef<FMediaViewerDelegates>& InDelegates)
 {
 	IMediaViewerLibraryWidget::FArgs Args;
 
 	ChildSlot
 	[
-		SAssignNew(Library, SMediaViewerLibrary, Args)
+		SAssignNew(Library, SMediaViewerLibrary, Args, InDelegates)
 	];
 
 	if (!::UMediaViewerLibraryIni::Get().HasGroup(UE::MediaViewer::Private::ColorGroupId))
diff --git a/Engine/Plugins/Experimental/MediaViewer/Source/MediaViewer/Private/Widgets/SMediaViewerLibraryPrivate.h b/Engine/Plugins/Experimental/MediaViewer/Source/MediaViewer/Private/Widgets/SMediaViewerLibraryPrivate.h
--- a/Engine/Plugins/Experimental/MediaViewer/Source/MediaViewer/Private/Widgets/SMediaViewerLibraryPrivate.h
+++ b/Engine/Plugins/Experimental/MediaViewer/Source/MediaViewer/Private/Widgets/SMediaViewerLibraryPrivate.h
@@ -31,7 +31,7 @@ public:
 	SMediaViewerLibraryPrivate();
 	virtual ~SMediaViewerLibraryPrivate() = default;
 
-	void Construct(const FArguments& InArgs, const TSharedPtr<FMediaViewerDelegates>& InDelegates);
+	void Construct(const FArguments& InArgs, const TSharedRef<FMediaViewerDelegates>& InDelegates);
 
 	TSharedRef<FMediaViewerLibrary> GetLibrary() const;
 
diff --git a/Engine/Plugins/Experimental/MediaViewer/Source/MediaViewer/Private/Widgets/SMediaViewerToolbar.cpp b/Engine/Plugins/Experimental/MediaViewer/Source/MediaViewer/Private/Widgets/SMediaViewerToolbar.cpp
--- a/Engine/Plugins/Experimental/MediaViewer/Source/MediaViewer/Private/Widgets/SMediaViewerToolbar.cpp
+++ b/Engine/Plugins/Experimental/MediaViewer/Source/MediaViewer/Private/Widgets/SMediaViewerToolbar.cpp
@@ -13,7 +13,6 @@
 #include "MediaViewerCommands.h"
 #include "MediaViewerStyle.h"
 #include "MediaViewerUtils.h"
-#include "SAssetDropTarget.h"
 #include "SEditorViewportToolBarMenu.h"
 #include "Styling/SlateTypes.h"
 #include "Styling/ToolBarStyle.h"
@@ -36,7 +35,7 @@ void SMediaViewerToolbar::PrivateRegisterAttributes(FSlateAttributeDescriptor::F
 {
 }
 
-void SMediaViewerToolbar::Construct(const FArguments& InArgs, const TSharedPtr<FMediaViewerDelegates>& InDelegates)
+void SMediaViewerToolbar::Construct(const FArguments& InArgs, const TSharedRef<FMediaViewerDelegates>& InDelegates)
 {
 	using namespace UE::MediaViewer;
 
@@ -220,7 +219,7 @@ TSharedRef<SWidget> SMediaViewerToolbar::MakeSideToolbar(EMediaImageViewerPositi
 		ImageDetails[static_cast<int32>(InPosition)] = SNew(SBox)
 			.Padding(3.f)
 			[
-				SNew(SMediaImageViewerDetails, InPosition, Delegates)
+				SNew(SMediaImageViewerDetails, InPosition, Delegates.ToSharedRef())
 			];
 	}
 
@@ -265,51 +264,6 @@ TSharedRef<SWidget> SMediaViewerToolbar::MakeSideToolbar(EMediaImageViewerPositi
 	return ToolbarBuilder.MakeWidget();
 }
 
-bool SMediaViewerToolbar::OnAssetDraggedOver(TArrayView<FAssetData> InAssetData) const
-{
-	IMediaViewerModule& Module = IMediaViewerModule::Get();
-
-	for (const FAssetData& AssetData : InAssetData)
-	{
-		if (Module.HasFactoryFor(AssetData))
-		{
-			return true;
-		}
-	}
-
-	return false;
-}
-
-void SMediaViewerToolbar::OnAssetsDropped(const FDragDropEvent& InDragDropEvent, TArrayView<FAssetData> InAssetData)
-{
-	IMediaViewerModule& Module = IMediaViewerModule::Get();
-
-	int32 AssetIndex = INDEX_NONE;
-
-	for (int32 Index = 0; Index < InAssetData.Num(); ++Index)
-	{
-		if (Module.HasFactoryFor(InAssetData[Index]))
-		{
-			AssetIndex = Index;
-			break;
-		}
-	}
-
-	if (AssetIndex == INDEX_NONE)
-	{
-		return;
-	}
-
-	if (!Module.HasImage(EMediaImageViewerPosition::First))
-	{
-		Module.SetImage(EMediaImageViewerPosition::First, InAssetData[AssetIndex]);
-	}
-	else if (!Module.HasImage(EMediaImageViewerPosition::Second))
-	{
-		Module.SetImage(EMediaImageViewerPosition::Second, InAssetData[AssetIndex]);
-	}
-}
-
 FText SMediaViewerToolbar::GetScaleMenuLabel(EMediaImageViewerPosition InPosition) const
 {
 	if (TSharedPtr<FMediaImageViewer> ImageViewer = Delegates->GetImageViewer.Execute(InPosition))
diff --git a/Engine/Plugins/Experimental/MediaViewer/Source/MediaViewer/Private/Widgets/SMediaViewerToolbar.h b/Engine/Plugins/Experimental/MediaViewer/Source/MediaViewer/Private/Widgets/SMediaViewerToolbar.h
--- a/Engine/Plugins/Experimental/MediaViewer/Source/MediaViewer/Private/Widgets/SMediaViewerToolbar.h
+++ b/Engine/Plugins/Experimental/MediaViewer/Source/MediaViewer/Private/Widgets/SMediaViewerToolbar.h
@@ -35,7 +35,7 @@ public:
 
 	virtual ~SMediaViewerToolbar() override = default;
 
-	void Construct(const FArguments& InArgs, const TSharedPtr<FMediaViewerDelegates>& InDelegates);
+	void Construct(const FArguments& InArgs, const TSharedRef<FMediaViewerDelegates>& InDelegates);
 
 	//~ Begin FNotifyHook
 	virtual void NotifyPostChange(const FPropertyChangedEvent& InPropertyChangedEvent, FProperty* InPropertyThatChanged) override;
@@ -55,9 +55,6 @@ protected:
 
 	TSharedRef<SWidget> MakeViewMenu();
 
-	bool OnAssetDraggedOver(TArrayView<FAssetData> InAssetData) const;
-	void OnAssetsDropped(const FDragDropEvent& InDragDropEvent, TArrayView<FAssetData> InAssetData);
-
 	FText GetScaleMenuLabel(EMediaImageViewerPosition InPosition) const;
 	TSharedRef<SWidget> MakeScaleMenu(EMediaImageViewerPosition InPosition) const;
 	
diff --git a/Engine/Plugins/Experimental/MediaViewer/Source/MediaViewer/Public/ImageViewer/MediaImageViewer.h b/Engine/Plugins/Experimental/MediaViewer/Source/MediaViewer/Public/ImageViewer/MediaImageViewer.h
--- a/Engine/Plugins/Experimental/MediaViewer/Source/MediaViewer/Public/ImageViewer/MediaImageViewer.h
+++ b/Engine/Plugins/Experimental/MediaViewer/Source/MediaViewer/Public/ImageViewer/MediaImageViewer.h
@@ -174,7 +174,7 @@ protected:
 public:
 	MEDIAVIEWER_API FMediaImageViewer(const FMediaImageViewerInfo& InImageInfo);
 
-	MEDIAVIEWER_API virtual ~FMediaImageViewer();
+	virtual ~FMediaImageViewer() = default;
 
 	const FMediaImageViewerInfo& GetInfo() const { return ImageInfo; }
 
@@ -189,13 +189,11 @@ public:
 
 	virtual TSharedPtr<FMediaViewerLibraryItem> CreateLibraryItem() const = 0;
 
-	const FMediaImageViewerPanelSettings& GetPanelSettings() const { return *PanelSettings->Get(); }
-	FMediaImageViewerPanelSettings& GetPanelSettings() { return *PanelSettings->Get(); }
-	TSharedRef<TStructOnScope<FMediaImageViewerPanelSettings>> GetPanelSettingsOnScope() const { return PanelSettings; }
+	const FMediaImageViewerPanelSettings& GetPanelSettings() const { return PanelSettings; }
+	FMediaImageViewerPanelSettings& GetPanelSettings() { return PanelSettings; }
 
-	const FMediaImagePaintSettings& GetPaintSettings() const { return *PaintSettings->Get(); }
-	FMediaImagePaintSettings& GetPaintSettings() { return *PaintSettings->Get(); }
-	TSharedRef<TStructOnScope<FMediaImagePaintSettings>> GetPaintSettingsOnScope() const { return PaintSettings; }
+	const FMediaImagePaintSettings& GetPaintSettings() const { return PaintSettings; }
+	FMediaImagePaintSettings& GetPaintSettings() { return PaintSettings; }
 
 	virtual TOptional<TVariant<FColor, FLinearColor>> GetPixelColor(const FIntPoint& InPixelCoords, int32 InMipLevel) const = 0;
 
@@ -281,8 +279,8 @@ protected:
 	static const FSlateColorBrush BackgroundColorBrush;
 
 	FMediaImageViewerInfo ImageInfo;
-	TSharedRef<TStructOnScope<FMediaImageViewerPanelSettings>> PanelSettings;
-	TSharedRef<TStructOnScope<FMediaImagePaintSettings>> PaintSettings;
+	FMediaImageViewerPanelSettings PanelSettings;
+	FMediaImagePaintSettings PaintSettings;
 	FSlateImageBrush BackgroundImageBrush;
 	TSharedPtr<FSlateBrush> Brush;
 
