Media Viewer: Toolbar now has the correct background color

diff --git a/Engine/Plugins/Experimental/MediaViewer/Source/MediaViewer/Private/Widgets/SMediaViewer.cpp b/Engine/Plugins/Experimental/MediaViewer/Source/MediaViewer/Private/Widgets/SMediaViewer.cpp
--- a/Engine/Plugins/Experimental/MediaViewer/Source/MediaViewer/Private/Widgets/SMediaViewer.cpp
+++ b/Engine/Plugins/Experimental/MediaViewer/Source/MediaViewer/Private/Widgets/SMediaViewer.cpp
@@ -12,6 +12,7 @@
 #include "MediaViewerCommands.h"
 #include "Sidebar/SSidebar.h"
 #include "Sidebar/SSidebarContainer.h"
+#include "Widgets/Colors/SColorBlock.h"
 #include "Widgets/Layout/SBox.h"
 #include "Widgets/Layout/SSplitter.h"
 #include "Widgets/MediaViewerDelegates.h"
@@ -509,7 +510,16 @@ TSharedRef<SMediaImageViewerOverlay> SMediaViewer::CreateOverlay(EMediaImageView
 
 TSharedRef<SWidget> SMediaViewer::CreateToolbar()
 {
-	return SNew(SMediaViewerToolbar, Delegates.ToSharedRef());
+	return SNew(SOverlay)
+		+ SOverlay::Slot()
+		[
+			SNew(SColorBlock)
+			.Color(FStyleColors::Panel.GetSpecifiedColor())
+		]
+		+ SOverlay::Slot()
+		[
+			SNew(SMediaViewerToolbar, Delegates.ToSharedRef())
+		];
 }
 
 void SMediaViewer::CreateView()
