Media Viewer: Library update
- Buttons on library are now only visible on mouse over, or when active in A or B.
- When items are active in the library (in either A or B) the button will appear selected.
- Removed render target dynamic group.

diff --git a/Engine/Plugins/Experimental/MediaViewer/Source/MediaViewer/Private/Widgets/MediaViewerDelegates.h b/Engine/Plugins/Experimental/MediaViewer/Source/MediaViewer/Private/Widgets/MediaViewerDelegates.h
--- a/Engine/Plugins/Experimental/MediaViewer/Source/MediaViewer/Private/Widgets/MediaViewerDelegates.h
+++ b/Engine/Plugins/Experimental/MediaViewer/Source/MediaViewer/Private/Widgets/MediaViewerDelegates.h
@@ -134,6 +134,9 @@ struct FMediaViewerDelegates
 	/** Sets the image viewer in the given position. */
 	FSetImageViewer SetImageViewer;
 
+	/** Sets the image viewer in the given position to the null viewer. */
+	FSimpleForPosition ClearImageViewer;
+
 	/** Gets the pixel coordinates hovered by the mouse for the image viewer in the given position. */
 	FGetLocationForPosition GetPixelCoordinates;
 
diff --git a/Engine/Plugins/Experimental/MediaViewer/Source/MediaViewer/Private/Widgets/SMediaViewer.cpp b/Engine/Plugins/Experimental/MediaViewer/Source/MediaViewer/Private/Widgets/SMediaViewer.cpp
--- a/Engine/Plugins/Experimental/MediaViewer/Source/MediaViewer/Private/Widgets/SMediaViewer.cpp
+++ b/Engine/Plugins/Experimental/MediaViewer/Source/MediaViewer/Private/Widgets/SMediaViewer.cpp
@@ -155,6 +155,13 @@ void SMediaViewer::SetImageViewer(EMediaImageViewerPosition InPosition, const TS
 	InvalidateView();
 }
 
+void SMediaViewer::ClearImageViewer(EMediaImageViewerPosition InPosition)
+{
+	ImageViewers[static_cast<int32>(InPosition)] = FNullImageViewer::GetNullImageViewer();
+
+	InvalidateView();
+}
+
 void SMediaViewer::BindCommands()
 {
 	const FMediaViewerCommands& Commands = FMediaViewerCommands::Get();
@@ -496,6 +503,7 @@ void SMediaViewer::CreateDelegates()
 	// Per panel delegates
 	Delegates->GetImageViewer.BindSP(this, &SMediaViewer::GetImageViewer);
 	Delegates->SetImageViewer.BindSP(this, &SMediaViewer::SetImageViewer);
+	Delegates->ClearImageViewer.BindSP(this, &SMediaViewer::ClearImageViewer);
 	Delegates->GetPixelCoordinates.BindSP(this, &SMediaViewer::GetPixelCoordinates);
 	Delegates->IsOverImage.BindSP(this, &SMediaViewer::IsOverImage);
 	Delegates->CopyTransformToAll.BindSP(this, &SMediaViewer::CopyTransformToAll);
diff --git a/Engine/Plugins/Experimental/MediaViewer/Source/MediaViewer/Private/Widgets/SMediaViewer.h b/Engine/Plugins/Experimental/MediaViewer/Source/MediaViewer/Private/Widgets/SMediaViewer.h
--- a/Engine/Plugins/Experimental/MediaViewer/Source/MediaViewer/Private/Widgets/SMediaViewer.h
+++ b/Engine/Plugins/Experimental/MediaViewer/Source/MediaViewer/Private/Widgets/SMediaViewer.h
@@ -68,6 +68,8 @@ public:
 	TSharedPtr<FMediaImageViewer> GetImageViewer(EMediaImageViewerPosition InPosition) const;
 
 	void SetImageViewer(EMediaImageViewerPosition InPosition, const TSharedRef<FMediaImageViewer>& InImageViewer);
+	
+	void ClearImageViewer(EMediaImageViewerPosition InPosition);
 
 	void SwapABImageViewers();
 
diff --git a/Engine/Plugins/Experimental/MediaViewer/Source/MediaViewer/Private/Widgets/SMediaViewerLibraryItem.cpp b/Engine/Plugins/Experimental/MediaViewer/Source/MediaViewer/Private/Widgets/SMediaViewerLibraryItem.cpp
--- a/Engine/Plugins/Experimental/MediaViewer/Source/MediaViewer/Private/Widgets/SMediaViewerLibraryItem.cpp
+++ b/Engine/Plugins/Experimental/MediaViewer/Source/MediaViewer/Private/Widgets/SMediaViewerLibraryItem.cpp
@@ -6,16 +6,19 @@
 #include "Containers/ArrayView.h"
 #include "DetailLayoutBuilder.h"
 #include "DragAndDrop/AssetDragDropOp.h"
+#include "ImageViewer/MediaImageViewer.h"
 #include "IMediaViewerModule.h"
 #include "Library/MediaViewerLibrary.h"
 #include "Library/MediaViewerLibraryGroup.h"
 #include "Library/MediaViewerLibraryItem.h"
 #include "MediaViewer.h"
 #include "MediaViewerStyle.h"
+#include "Styling/AppStyle.h"
 #include "Styling/StyleColors.h"
 #include "Widgets/Colors/SColorBlock.h"
 #include "Widgets/Images/SImage.h"
 #include "Widgets/Input/SButton.h"
+#include "Widgets/Input/SCheckBox.h"
 #include "Widgets/MediaViewerDelegates.h"
 #include "Widgets/MediaViewerLibraryItemDragDropOperation.h"
 #include "Widgets/SBoxPanel.h"
@@ -116,16 +119,19 @@ void SMediaViewerLibraryItem::Construct(const FArguments& InArgs, const TSharedR
 		];
 
 	const FButtonStyle& ButtonStyle = FMediaViewerStyle::Get().GetWidgetStyle<FButtonStyle>("LibraryButtonStyle");
+	const FCheckBoxStyle& CheckBoxStyle = FAppStyle::Get().GetWidgetStyle<FCheckBoxStyle>("DetailsView.SectionButton");
 
 	Box->AddSlot()
 		.AutoWidth()
 		.VAlign(VAlign_Center)
 		.Padding(3.f, 5.f, 0.f, 5.f)
 		[
-			SNew(SButton)
-			.ContentPadding(5.f)
-			.OnClicked(this, &SMediaViewerLibraryItem::OnUseButtonClicked, EMediaImageViewerPosition::First)
-			.ButtonStyle(&ButtonStyle)
+			SNew(SCheckBox)
+			.Style(&CheckBoxStyle)
+			.Visibility(this, &SMediaViewerLibraryItem::GetHoveredOrActiveVisibility)
+			.Padding(FMargin(9.f, 5.f))
+			.IsChecked(this, &SMediaViewerLibraryItem::IsActiveState, EMediaImageViewerPosition::First)
+			.OnCheckStateChanged(this, &SMediaViewerLibraryItem::OnUseButtonClicked, EMediaImageViewerPosition::First)
 			[
 				SNew(STextBlock)
 				.Text(INVTEXT("A"))
@@ -139,10 +145,12 @@ void SMediaViewerLibraryItem::Construct(const FArguments& InArgs, const TSharedR
 		.Padding(3.f, 5.f, 3.f, 5.f)
 		.VAlign(VAlign_Center)
 		[
-			SNew(SButton)
-			.ContentPadding(5.f)
-			.OnClicked(this, &SMediaViewerLibraryItem::OnUseButtonClicked, EMediaImageViewerPosition::Second)
-			.ButtonStyle(&ButtonStyle)
+			SNew(SCheckBox)
+			.Style(&CheckBoxStyle)
+			.Visibility(this, &SMediaViewerLibraryItem::GetHoveredOrActiveVisibility)
+			.Padding(FMargin(9.f, 5.f))
+			.IsChecked(this, &SMediaViewerLibraryItem::IsActiveState, EMediaImageViewerPosition::Second)
+			.OnCheckStateChanged(this, &SMediaViewerLibraryItem::OnUseButtonClicked, EMediaImageViewerPosition::Second)
 			[
 				SNew(STextBlock)
 				.Text(INVTEXT("B"))
@@ -159,6 +167,7 @@ void SMediaViewerLibraryItem::Construct(const FArguments& InArgs, const TSharedR
 			.Padding(0.f, 5.f, 3.f, 5.f)
 			[
 				SNew(SButton)
+				.Visibility(this, &SMediaViewerLibraryItem::GetHoveredOrActiveVisibility)
 				.ContentPadding(FMargin(3.f, 5.f, 3.f, 6.f))
 				.OnClicked(this, &SMediaViewerLibraryItem::OnRemoveButtonClicked)
 				.ButtonStyle(&ButtonStyle)
@@ -194,40 +203,92 @@ void SMediaViewerLibraryItem::Construct(const FArguments& InArgs, const TSharedR
 	);
 }
 
-FReply SMediaViewerLibraryItem::OnUseButtonClicked(EMediaImageViewerPosition InPosition)
+EVisibility SMediaViewerLibraryItem::GetHoverVisibility() const
 {
-	TSharedPtr<FMediaViewerLibrary> Library = LibraryWeak.Pin();
+	if (IsHovered())
+	{
+		return EVisibility::Visible;
+	}
 
-	if (!Library.IsValid())
+	return EVisibility::Collapsed;
+}
+
+bool SMediaViewerLibraryItem::IsActive(EMediaImageViewerPosition InPosition) const
+{
+	if (TSharedPtr<FMediaImageViewer> ImageViewer = Delegates->GetImageViewer.Execute(InPosition))
 	{
-		return FReply::Handled();
+		return ImageViewer->GetInfo().Id == Item->GetId();
+	}
+
+	return false;
+}
+
+EVisibility SMediaViewerLibraryItem::GetHoveredOrActiveVisibility() const
+{
+	if (IsHovered())
+	{
+		return EVisibility::Visible;
 	}
 
+	for (int32 Index = 0; Index < static_cast<int32>(EMediaImageViewerPosition::COUNT); ++Index)
+	{
+		if (IsActive(static_cast<EMediaImageViewerPosition>(Index)))
+		{
+			return EVisibility::Visible;
+		}
+	}
+
+	return EVisibility::Collapsed;
+}
+
+void SMediaViewerLibraryItem::SetImageViewer(EMediaImageViewerPosition InPosition)
+{
 	if (!Item.IsValid())
 	{
-		return FReply::Handled();
+		return;
 	}
 
 	TSharedPtr<FMediaImageViewer> ImageViewer = Item->CreateImageViewer();
-	
+
 	if (!ImageViewer.IsValid())
 	{
-		return FReply::Handled();
+		return;
 	}
 
 	Delegates->SetImageViewer.Execute(InPosition, ImageViewer.ToSharedRef());
+}
 
-	return FReply::Handled();
+void SMediaViewerLibraryItem::ClearImageViewer(EMediaImageViewerPosition InPosition)
+{
+	Delegates->ClearImageViewer.Execute(InPosition);
 }
 
-EVisibility SMediaViewerLibraryItem::GetHoverVisibility() const
+ECheckBoxState SMediaViewerLibraryItem::IsActiveState(EMediaImageViewerPosition InPosition) const
 {
-	if (IsHovered())
+	if (IsActive(InPosition))
 	{
-		return EVisibility::Visible;
+		return ECheckBoxState::Checked;
 	}
 
-	return EVisibility::Hidden;
+	return ECheckBoxState::Unchecked;
+}
+
+void SMediaViewerLibraryItem::OnUseButtonClicked(ECheckBoxState InCheckState, EMediaImageViewerPosition InPosition)
+{
+	switch (InCheckState)
+	{
+		case ECheckBoxState::Checked:
+			SetImageViewer(InPosition);
+			break;
+
+		case ECheckBoxState::Unchecked:
+			ClearImageViewer(InPosition);
+			break;
+
+		default:
+			// Do nothing
+			break;
+	}
 }
 
 FReply SMediaViewerLibraryItem::OnRemoveButtonClicked()
diff --git a/Engine/Plugins/Experimental/MediaViewer/Source/MediaViewer/Private/Widgets/SMediaViewerLibraryItem.h b/Engine/Plugins/Experimental/MediaViewer/Source/MediaViewer/Private/Widgets/SMediaViewerLibraryItem.h
--- a/Engine/Plugins/Experimental/MediaViewer/Source/MediaViewer/Private/Widgets/SMediaViewerLibraryItem.h
+++ b/Engine/Plugins/Experimental/MediaViewer/Source/MediaViewer/Private/Widgets/SMediaViewerLibraryItem.h
@@ -57,10 +57,19 @@ protected:
 	TSharedPtr<FMediaViewerLibraryItem> Item;
 	TSharedPtr<FSlateBrush> ThumbnailBrush;
 
-	FReply OnUseButtonClicked(EMediaImageViewerPosition InPosition);
-
 	EVisibility GetHoverVisibility() const;
 
+	bool IsActive(EMediaImageViewerPosition InPosition) const;
+
+	EVisibility GetHoveredOrActiveVisibility() const;
+
+	void SetImageViewer(EMediaImageViewerPosition InPosition);
+
+	void ClearImageViewer(EMediaImageViewerPosition InPosition);
+
+	ECheckBoxState IsActiveState(EMediaImageViewerPosition InPosition) const;
+	void OnUseButtonClicked(ECheckBoxState InCheckState, EMediaImageViewerPosition InPosition);
+
 	FReply OnRemoveButtonClicked();
 
 	FReply OnDragDetected(const FGeometry& InMyGeometry, const FPointerEvent& InPointerEvent);
diff --git a/Engine/Plugins/Experimental/MediaViewer/Source/MediaViewer/Private/Widgets/SMediaViewerLibraryPrivate.cpp b/Engine/Plugins/Experimental/MediaViewer/Source/MediaViewer/Private/Widgets/SMediaViewerLibraryPrivate.cpp
--- a/Engine/Plugins/Experimental/MediaViewer/Source/MediaViewer/Private/Widgets/SMediaViewerLibraryPrivate.cpp
+++ b/Engine/Plugins/Experimental/MediaViewer/Source/MediaViewer/Private/Widgets/SMediaViewerLibraryPrivate.cpp
@@ -10,7 +10,6 @@
 #include "Library/MediaViewerLibraryGroup.h"
 #include "Library/MediaViewerLibraryIni.h"
 #include "Library/MediaViewerLibraryItem.h"
-#include "Library/TextureRenderTarget2DGroup.h"
 #include "MediaViewer.h"
 #include "Misc/Optional.h"
 #include "Widgets/MediaViewerDelegates.h"
@@ -53,7 +52,6 @@ void SMediaViewerLibraryPrivate::Construct(const FArguments& InArgs, const TShar
 	TSharedRef<FMediaViewerLibrary> LibraryImpl = GetLibrary();
 
 	LibraryImpl->AddGroup(MakeShared<FMediaTextureGroup>(LibraryImpl));
-	LibraryImpl->AddGroup(MakeShared<FTextureRenderTarget2DGroup>(LibraryImpl));
 	LibraryImpl->AddGroup(MakeShared<FLevelEditorViewportGroup>(LibraryImpl));
 }
 
