Media Viewer: Added A/B labels to each side

diff --git a/Engine/Plugins/Experimental/MediaViewer/Source/MediaViewer/Private/Widgets/SMediaImageViewerOverlay.cpp b/Engine/Plugins/Experimental/MediaViewer/Source/MediaViewer/Private/Widgets/SMediaImageViewerOverlay.cpp
--- a/Engine/Plugins/Experimental/MediaViewer/Source/MediaViewer/Private/Widgets/SMediaImageViewerOverlay.cpp
+++ b/Engine/Plugins/Experimental/MediaViewer/Source/MediaViewer/Private/Widgets/SMediaImageViewerOverlay.cpp
@@ -414,11 +414,33 @@ TSharedRef<SWidget> SMediaImageViewerOverlay::CreateOverlay(bool bInCreateSplitD
 
 	if (!bInCreateSplitDropTarget)
 	{
-		return SNew(SMediaViewerDropTarget, Delegates.ToSharedRef())
+		TSharedRef<SWidget> PanelName = SNew(SBox)
+			.HAlign(Position == EMediaImageViewerPosition::First ? HAlign_Right : HAlign_Left)
+			.VAlign(VAlign_Top)
+			.Padding(5.f)
+			[
+				SNew(STextBlock)
+					.ColorAndOpacity(FStyleColors::Foreground.GetSpecifiedColor())
+					.ShadowColorAndOpacity(FStyleColors::Panel.GetSpecifiedColor())
+					.ShadowOffset(FVector2D(1.f, 1.f))
+					.Text(Position == EMediaImageViewerPosition::First ? INVTEXT("A") : INVTEXT("B"))
+			];
+
+		TSharedRef<SWidget> Overlay = SNew(SMediaViewerDropTarget, Delegates.ToSharedRef())
 			.Position(Position)
 			[
 				ImageViewerOverlay.IsValid() ? ImageViewerOverlay.ToSharedRef() : SNullWidget::NullWidget
 			];
+
+		return SNew(SOverlay)
+			+ SOverlay::Slot()
+			[
+				Overlay
+			]
+			+ SOverlay::Slot()
+			[
+				PanelName
+			];
 	}
 
 	TSharedPtr<SWidget> SplitDropTarget;
@@ -470,7 +492,7 @@ TSharedRef<SWidget> SMediaImageViewerOverlay::CreateOverlay(bool bInCreateSplitD
 		+ SOverlay::Slot()
 		[
 			ImageViewerOverlay.IsValid() ? ImageViewerOverlay.ToSharedRef() : SNullWidget::NullWidget
-		];		
+		];
 }
 
 TSharedRef<SWidget> SMediaImageViewerOverlay::CreateStatusBar(const TSharedRef<FMediaViewerDelegates>& InDelegates)
