Media Viewer: Min size of slider changed from 50 to 10 pixels

diff --git a/Engine/Plugins/Experimental/MediaViewer/Source/MediaViewer/Private/Widgets/SMediaViewer.cpp b/Engine/Plugins/Experimental/MediaViewer/Source/MediaViewer/Private/Widgets/SMediaViewer.cpp
--- a/Engine/Plugins/Experimental/MediaViewer/Source/MediaViewer/Private/Widgets/SMediaViewer.cpp
+++ b/Engine/Plugins/Experimental/MediaViewer/Source/MediaViewer/Private/Widgets/SMediaViewer.cpp
@@ -597,7 +597,7 @@ void SMediaViewer::CreateABView()
 
 		+ SSplitter::Slot()
 		.Value(0.5f)
-		.MinSize(50.f)
+		.MinSize(10.f)
 		.Resizable(true)
 		.SizeRule(SSplitter::FractionOfParent)
 		.OnSlotResized(this, &SMediaViewer::OnABResized)
@@ -607,7 +607,7 @@ void SMediaViewer::CreateABView()
 
 		+ SSplitter::Slot()
 		.Value(0.5f)
-		.MinSize(50.f)
+		.MinSize(10.f)
 		.Resizable(true)
 		.SizeRule(SSplitter::FractionOfParent)
 		[
