Media Viewer: Displaying names of the images.

diff --git a/Engine/Plugins/Experimental/MediaViewer/Source/MediaViewer/Private/ImageViewers/ColorImageViewer.cpp b/Engine/Plugins/Experimental/MediaViewer/Source/MediaViewer/Private/ImageViewers/ColorImageViewer.cpp
--- a/Engine/Plugins/Experimental/MediaViewer/Source/MediaViewer/Private/ImageViewers/ColorImageViewer.cpp
+++ b/Engine/Plugins/Experimental/MediaViewer/Source/MediaViewer/Private/ImageViewers/ColorImageViewer.cpp
@@ -96,10 +96,10 @@ TSharedPtr<FMediaImageViewer> FColorImageViewer::FItem::CreateImageViewer() cons
 
 	if (Id.IsValid())
 	{
-		return MakeShared<FColorImageViewer>(Id, LoadFromString(StringValue));
+		return MakeShared<FColorImageViewer>(Id, LoadFromString(StringValue), Name);
 	}
 	
-	return MakeShared<FColorImageViewer>(LoadFromString(StringValue));
+	return MakeShared<FColorImageViewer>(LoadFromString(StringValue), Name);
 }
 
 TSharedPtr<FMediaViewerLibraryItem> FColorImageViewer::FItem::Clone() const
@@ -108,26 +108,21 @@ TSharedPtr<FMediaViewerLibraryItem> FColorImageViewer::FItem::Clone() const
 }
 
 FColorImageViewer::FColorImageViewer()
-	: FColorImageViewer(UE::MediaViewer::Private::MakeRandomColor())
+	: FColorImageViewer(UE::MediaViewer::Private::MakeRandomColor(), LOCTEXT("Color", "Color"))
 {
 }
 
-FColorImageViewer::FColorImageViewer(const FLinearColor& InColor)
-	: FColorImageViewer(FGuid::NewGuid(), InColor)
+FColorImageViewer::FColorImageViewer(const FLinearColor& InColor, const FText& InDisplayName)
+	: FColorImageViewer(FGuid::NewGuid(), InColor, InDisplayName)
 {
 }
 
-FColorImageViewer::FColorImageViewer(const FGuid& InId)
-	: FColorImageViewer(InId, UE::MediaViewer::Private::MakeRandomColor())
-{
-}
-
-FColorImageViewer::FColorImageViewer(const FGuid& InId, const FLinearColor& InColor)
+FColorImageViewer::FColorImageViewer(const FGuid& InId, const FLinearColor& InColor, const FText& InDisplayName)
 	: FMediaImageViewer({
 		InId,
 		{UE::MediaViewer::Private::BlockSize, UE::MediaViewer::Private::BlockSize},
 		1,
-		LOCTEXT("Color", "Color")
+		InDisplayName
 	})
 {
 	Brush = ColorBrush;
diff --git a/Engine/Plugins/Experimental/MediaViewer/Source/MediaViewer/Private/ImageViewers/ColorImageViewer.h b/Engine/Plugins/Experimental/MediaViewer/Source/MediaViewer/Private/ImageViewers/ColorImageViewer.h
--- a/Engine/Plugins/Experimental/MediaViewer/Source/MediaViewer/Private/ImageViewers/ColorImageViewer.h
+++ b/Engine/Plugins/Experimental/MediaViewer/Source/MediaViewer/Private/ImageViewers/ColorImageViewer.h
@@ -58,9 +58,8 @@ public:
 	static const FLazyName ItemTypeName;
 
 	FColorImageViewer();
-	FColorImageViewer(const FLinearColor& InColor);
-	FColorImageViewer(const FGuid& InId);
-	FColorImageViewer(const FGuid& InId, const FLinearColor& InColor);
+	FColorImageViewer(const FLinearColor& InColor, const FText& InDisplayName);
+	FColorImageViewer(const FGuid& InId, const FLinearColor& InColor, const FText& InDisplayName);
 
 	virtual ~FColorImageViewer() = default;
 
diff --git a/Engine/Plugins/Experimental/MediaViewer/Source/MediaViewer/Private/Widgets/SMediaImageViewerOverlay.cpp b/Engine/Plugins/Experimental/MediaViewer/Source/MediaViewer/Private/Widgets/SMediaImageViewerOverlay.cpp
--- a/Engine/Plugins/Experimental/MediaViewer/Source/MediaViewer/Private/Widgets/SMediaImageViewerOverlay.cpp
+++ b/Engine/Plugins/Experimental/MediaViewer/Source/MediaViewer/Private/Widgets/SMediaImageViewerOverlay.cpp
@@ -412,6 +412,18 @@ TSharedRef<SWidget> SMediaImageViewerOverlay::CreateOverlay(bool bInCreateSplitD
 
 	TSharedPtr<SWidget> ImageViewerOverlay = ImageViewer.IsValid() ? ImageViewer->GetOverlayWidget(Position) : nullptr;
 
+	TSharedRef<SWidget> DisplayName = SNew(SBox)
+		.HAlign(HAlign_Center)
+		.VAlign(VAlign_Top)
+		.Padding(5.f)
+		[
+			SNew(STextBlock)
+			.ColorAndOpacity(FStyleColors::Foreground.GetSpecifiedColor())
+			.ShadowColorAndOpacity(FStyleColors::Panel.GetSpecifiedColor())
+			.ShadowOffset(FVector2D(1.f, 1.f))
+			.Text(ImageViewer->GetInfo().DisplayName)
+		];
+
 	if (!bInCreateSplitDropTarget)
 	{
 		TSharedRef<SWidget> PanelName = SNew(SBox)
@@ -420,10 +432,10 @@ TSharedRef<SWidget> SMediaImageViewerOverlay::CreateOverlay(bool bInCreateSplitD
 			.Padding(5.f)
 			[
 				SNew(STextBlock)
-					.ColorAndOpacity(FStyleColors::Foreground.GetSpecifiedColor())
-					.ShadowColorAndOpacity(FStyleColors::Panel.GetSpecifiedColor())
-					.ShadowOffset(FVector2D(1.f, 1.f))
-					.Text(Position == EMediaImageViewerPosition::First ? INVTEXT("A") : INVTEXT("B"))
+				.ColorAndOpacity(FStyleColors::Foreground.GetSpecifiedColor())
+				.ShadowColorAndOpacity(FStyleColors::Panel.GetSpecifiedColor())
+				.ShadowOffset(FVector2D(1.f, 1.f))
+				.Text(Position == EMediaImageViewerPosition::First ? INVTEXT("A") : INVTEXT("B"))
 			];
 
 		TSharedRef<SWidget> Overlay = SNew(SMediaViewerDropTarget, Delegates.ToSharedRef())
@@ -440,6 +452,10 @@ TSharedRef<SWidget> SMediaImageViewerOverlay::CreateOverlay(bool bInCreateSplitD
 			+ SOverlay::Slot()
 			[
 				PanelName
+			]
+			+ SOverlay::Slot()
+			[
+				DisplayName
 			];
 	}
 
@@ -492,6 +508,10 @@ TSharedRef<SWidget> SMediaImageViewerOverlay::CreateOverlay(bool bInCreateSplitD
 		+ SOverlay::Slot()
 		[
 			ImageViewerOverlay.IsValid() ? ImageViewerOverlay.ToSharedRef() : SNullWidget::NullWidget
+		]
+		+ SOverlay::Slot()
+		[
+			DisplayName
 		];
 }
 
