Media Viewer: Reworked View system
- Reworked how single and A/B views work

diff --git a/Engine/Plugins/Experimental/MediaViewer/Source/MediaViewer/Private/ImageViewers/NullImageViewer.cpp b/Engine/Plugins/Experimental/MediaViewer/Source/MediaViewer/Private/ImageViewers/NullImageViewer.cpp
--- a/Engine/Plugins/Experimental/MediaViewer/Source/MediaViewer/Private/ImageViewers/NullImageViewer.cpp
+++ b/Engine/Plugins/Experimental/MediaViewer/Source/MediaViewer/Private/ImageViewers/NullImageViewer.cpp
@@ -23,7 +23,7 @@ FNullImageViewer::FNullImageViewer()
 		FGuid(), 
 		FIntPoint::ZeroValue, 
 		0, 
-		LOCTEXT("Null", "Null")
+		LOCTEXT("Null", "-")
 	})
 {
 }
diff --git a/Engine/Plugins/Experimental/MediaViewer/Source/MediaViewer/Private/Widgets/MediaViewerDelegates.h b/Engine/Plugins/Experimental/MediaViewer/Source/MediaViewer/Private/Widgets/MediaViewerDelegates.h
--- a/Engine/Plugins/Experimental/MediaViewer/Source/MediaViewer/Private/Widgets/MediaViewerDelegates.h
+++ b/Engine/Plugins/Experimental/MediaViewer/Source/MediaViewer/Private/Widgets/MediaViewerDelegates.h
@@ -29,7 +29,7 @@ struct FMediaViewerDelegates
 	DECLARE_DELEGATE_RetVal(FIntPoint, FGetLocation)
 	DECLARE_DELEGATE_RetVal(EOrientation, FGetOrientation)
 	DECLARE_DELEGATE_OneParam(FSetOrientation, EOrientation)
-	DECLARE_DELEGATE_RetVal(EMediaImageViewerActivePosition, FGetActiveViews)
+	DECLARE_DELEGATE_RetVal(EMediaImageViewerActivePosition, FGetActiveView)
 	DECLARE_DELEGATE_RetVal(bool, FGetBool)
 	DECLARE_DELEGATE_RetVal(float, FGetFloat)
 	DECLARE_DELEGATE_OneParam(FSetFloat, float)
@@ -62,7 +62,7 @@ struct FMediaViewerDelegates
 	FSetOrientation SetABOrientation;
 
 	/** Returns a bitflag value indicating which image viewers are currently active (have images, not null or null image viewer. */
-	FGetActiveViews GetActiveViews;
+	FGetActiveView GetActiveView;
 
 	/* Gets a pointer to the Media Viewer configurable settings. */
 	FGetSettings GetSettings;
diff --git a/Engine/Plugins/Experimental/MediaViewer/Source/MediaViewer/Private/Widgets/SMediaImageViewerOverlay.cpp b/Engine/Plugins/Experimental/MediaViewer/Source/MediaViewer/Private/Widgets/SMediaImageViewerOverlay.cpp
--- a/Engine/Plugins/Experimental/MediaViewer/Source/MediaViewer/Private/Widgets/SMediaImageViewerOverlay.cpp
+++ b/Engine/Plugins/Experimental/MediaViewer/Source/MediaViewer/Private/Widgets/SMediaImageViewerOverlay.cpp
@@ -76,7 +76,7 @@ void SMediaImageViewerOverlay::Construct(const FArguments& InArgs, EMediaImageVi
 			.HAlign(HAlign_Fill)
 			.VAlign(VAlign_Fill)
 			[
-				CreateOverlay(InArgs._bCreateSplitDropTarget)
+				CreateOverlay(InArgs._bShowPanelName)
 			]
 
 			+ SVerticalBox::Slot()
@@ -401,7 +401,7 @@ void SMediaImageViewerOverlay::BindCommands()
 	);
 }
 
-TSharedRef<SWidget> SMediaImageViewerOverlay::CreateOverlay(bool bInCreateSplitDropTarget)
+TSharedRef<SWidget> SMediaImageViewerOverlay::CreateOverlay(bool bInShowPanelName)
 {
 	TSharedPtr<FMediaImageViewer> ImageViewer = Delegates->GetImageViewer.Execute(Position);
 
@@ -410,109 +410,59 @@ TSharedRef<SWidget> SMediaImageViewerOverlay::CreateOverlay(bool bInCreateSplitD
 		return SNullWidget::NullWidget;
 	}
 
+	TSharedRef<SOverlay> Overlay = SNew(SOverlay);
+
 	TSharedPtr<SWidget> ImageViewerOverlay = ImageViewer.IsValid() ? ImageViewer->GetOverlayWidget(Position) : nullptr;
 
-	TSharedRef<SWidget> DisplayName = SNew(SBox)
-		.HAlign(HAlign_Center)
-		.VAlign(VAlign_Top)
-		.Padding(5.f)
+	Overlay->AddSlot()
 		[
-			SNew(STextBlock)
-			.ColorAndOpacity(FStyleColors::Foreground.GetSpecifiedColor())
-			.ShadowColorAndOpacity(FStyleColors::Panel.GetSpecifiedColor())
-			.ShadowOffset(FVector2D(1.f, 1.f))
-			.Text(ImageViewer->GetInfo().DisplayName)
-		];
-
-	if (!bInCreateSplitDropTarget)
-	{
-		TSharedRef<SWidget> PanelName = SNew(SBox)
-			.HAlign(Position == EMediaImageViewerPosition::First ? HAlign_Right : HAlign_Left)
-			.VAlign(VAlign_Top)
-			.Padding(5.f)
-			[
-				SNew(STextBlock)
-				.ColorAndOpacity(FStyleColors::Foreground.GetSpecifiedColor())
-				.ShadowColorAndOpacity(FStyleColors::Panel.GetSpecifiedColor())
-				.ShadowOffset(FVector2D(1.f, 1.f))
-				.Text(Position == EMediaImageViewerPosition::First ? INVTEXT("A") : INVTEXT("B"))
-			];
-
-		TSharedRef<SWidget> Overlay = SNew(SMediaViewerDropTarget, Delegates.ToSharedRef())
+			SNew(SMediaViewerDropTarget, Delegates.ToSharedRef())
 			.Position(Position)
 			[
 				ImageViewerOverlay.IsValid() ? ImageViewerOverlay.ToSharedRef() : SNullWidget::NullWidget
-			];
-
-		return SNew(SOverlay)
-			+ SOverlay::Slot()
-			[
-				Overlay
-			]
-			+ SOverlay::Slot()
-			[
-				PanelName
-			]
-			+ SOverlay::Slot()
-			[
-				DisplayName
-			];
-	}
-
-	TSharedPtr<SWidget> SplitDropTarget;
-
-	if (Delegates->GetABOrientation.Execute() == Orient_Horizontal)
-	{
-		SplitDropTarget = SNew(SHorizontalBox)
-			+ SHorizontalBox::Slot()
-			.FillWidth(1.f)
-			.Padding(0.f, 0.f, 2.f, 0.f)
-			[
-				SNew(SMediaViewerDropTarget, Delegates.ToSharedRef())
-				.Position(EMediaImageViewerPosition::First)
 			]
+		];
 
-			+ SHorizontalBox::Slot()
-			.FillWidth(1.f)
-			.Padding(2.f, 0.f, 0.f, 0.f)
-			[
-				SNew(SMediaViewerDropTarget, Delegates.ToSharedRef())
-				.Position(EMediaImageViewerPosition::Second)
-			];
-	}
-	else
+	if (bInShowPanelName)
 	{
-		SplitDropTarget = SNew(SVerticalBox)
-			+ SVerticalBox::Slot()
-			.FillHeight(1.f)
-			.Padding(0.f, 0.f, 0.f, 2.f)
-			[
-				SNew(SMediaViewerDropTarget, Delegates.ToSharedRef())
-				.Position(EMediaImageViewerPosition::First)
-			]
+		const EHorizontalAlignment HAlign = (Position == EMediaImageViewerPosition::First
+			&& Delegates->GetActiveView.Execute() == EMediaImageViewerActivePosition::Both
+			&& Delegates->GetABOrientation.Execute() == Orient_Horizontal)
+			? HAlign_Right
+			: HAlign_Left;
 
-			+ SVerticalBox::Slot()
-			.FillHeight(1.f)
-			.Padding(0.f, 2.f, 0.f, 0.f)
+		Overlay->AddSlot()
 			[
-				SNew(SMediaViewerDropTarget, Delegates.ToSharedRef())
-				.Position(EMediaImageViewerPosition::Second)
+				SNew(SBox)
+				.HAlign(HAlign)
+				.VAlign(VAlign_Top)
+				.Padding(5.f)
+				[
+					SNew(STextBlock)
+					.ColorAndOpacity(FStyleColors::Foreground.GetSpecifiedColor())
+					.ShadowColorAndOpacity(FStyleColors::Panel.GetSpecifiedColor())
+					.ShadowOffset(FVector2D(1.f, 1.f))
+					.Text(Position == EMediaImageViewerPosition::First ? INVTEXT("A") : INVTEXT("B"))
+				]
 			];
 	}
 
-	return SNew(SOverlay)		
-		+ SOverlay::Slot()
+	Overlay->AddSlot()
 		[
-			SplitDropTarget.ToSharedRef()
-		]
-		+ SOverlay::Slot()
-		[
-			ImageViewerOverlay.IsValid() ? ImageViewerOverlay.ToSharedRef() : SNullWidget::NullWidget
-		]
-		+ SOverlay::Slot()
-		[
-			DisplayName
+			SNew(SBox)
+			.HAlign(HAlign_Center)
+			.VAlign(VAlign_Top)
+			.Padding(5.f)
+			[
+				SNew(STextBlock)
+				.ColorAndOpacity(FStyleColors::Foreground.GetSpecifiedColor())
+				.ShadowColorAndOpacity(FStyleColors::Panel.GetSpecifiedColor())
+				.ShadowOffset(FVector2D(1.f, 1.f))
+				.Text(ImageViewer->GetInfo().DisplayName)
+			]
 		];
+
+	return Overlay;
 }
 
 TSharedRef<SWidget> SMediaImageViewerOverlay::CreateStatusBar(const TSharedRef<FMediaViewerDelegates>& InDelegates)
@@ -704,7 +654,7 @@ void SMediaImageViewerOverlay::ScaleToFit()
 
 bool SMediaImageViewerOverlay::CanResetTransform()
 {
-	return Delegates->GetActiveViews.Execute() == EMediaImageViewerActivePosition::Both;
+	return Delegates->GetActiveView.Execute() == EMediaImageViewerActivePosition::Both;
 }
 
 void SMediaImageViewerOverlay::Try_ResetTransform()
@@ -735,7 +685,7 @@ void SMediaImageViewerOverlay::Try_SetTransform(FVector InOffset, FRotator InRot
 
 bool SMediaImageViewerOverlay::CanCopyTransform()
 {
-	return Delegates->GetActiveViews.Execute() == EMediaImageViewerActivePosition::Both;
+	return Delegates->GetActiveView.Execute() == EMediaImageViewerActivePosition::Both;
 }
 
 void SMediaImageViewerOverlay::CopyTransform()
diff --git a/Engine/Plugins/Experimental/MediaViewer/Source/MediaViewer/Private/Widgets/SMediaImageViewerOverlay.h b/Engine/Plugins/Experimental/MediaViewer/Source/MediaViewer/Private/Widgets/SMediaImageViewerOverlay.h
--- a/Engine/Plugins/Experimental/MediaViewer/Source/MediaViewer/Private/Widgets/SMediaImageViewerOverlay.h
+++ b/Engine/Plugins/Experimental/MediaViewer/Source/MediaViewer/Private/Widgets/SMediaImageViewerOverlay.h
@@ -28,9 +28,9 @@ class SMediaImageViewerOverlay : public SCompoundWidget
 	SLATE_DECLARE_WIDGET(SMediaImageViewerOverlay, SCompoundWidget)
 
 	SLATE_BEGIN_ARGS(SMediaImageViewerOverlay)
-		: _bCreateSplitDropTarget(false)
+		: _bShowPanelName(true)
 		{}
-		SLATE_ARGUMENT(bool, bCreateSplitDropTarget)
+		SLATE_ARGUMENT(bool, bShowPanelName)
 	SLATE_END_ARGS()
 
 public:
@@ -91,7 +91,7 @@ protected:
 
 	void BindCommands();
 
-	TSharedRef<SWidget> CreateOverlay(bool bInCreateSplitDropTarget);
+	TSharedRef<SWidget> CreateOverlay(bool bInShowPanelName);
 
 	TSharedRef<SWidget> CreateStatusBar(const TSharedRef<FMediaViewerDelegates>& InDelegates);
 
diff --git a/Engine/Plugins/Experimental/MediaViewer/Source/MediaViewer/Private/Widgets/SMediaViewer.cpp b/Engine/Plugins/Experimental/MediaViewer/Source/MediaViewer/Private/Widgets/SMediaViewer.cpp
--- a/Engine/Plugins/Experimental/MediaViewer/Source/MediaViewer/Private/Widgets/SMediaViewer.cpp
+++ b/Engine/Plugins/Experimental/MediaViewer/Source/MediaViewer/Private/Widgets/SMediaViewer.cpp
@@ -52,7 +52,8 @@ void SMediaViewer::Construct(const FArguments& InArgs, const FMediaViewerArgs& I
 	MediaViewerArgs = InMediaViewerArgs;
 	ImageViewers[static_cast<int32>(EMediaImageViewerPosition::First)] = InImageViewerFirst;
 	ImageViewers[static_cast<int32>(EMediaImageViewerPosition::Second)] = InImageViewerSecond;
-	ActiveViews = EMediaImageViewerActivePosition::None;
+	ActiveView = EMediaImageViewerActivePosition::Single;
+	RequestedView = EMediaImageViewerActivePosition::Single;
 	ContentSlot = nullptr;
 	CursorLocalPosition = {0, 0};
 
@@ -145,21 +146,31 @@ TSharedPtr<FMediaImageViewer> SMediaViewer::GetImageViewer(EMediaImageViewerPosi
 
 void SMediaViewer::SetImageViewer(EMediaImageViewerPosition InPosition, const TSharedRef<FMediaImageViewer>& InImageViewer)
 {
-	ImageViewers[static_cast<int32>(InPosition)] = InImageViewer;
+	const int32 Index = static_cast<int32>(InPosition);
+
+	if (ImageViewers[Index].IsValid() && ImageViewers[Index]->GetInfo().Id == InImageViewer->GetInfo().Id)
+	{
+		return;
+	}
+
+	ImageViewers[Index] = InImageViewer;
 
 	if (InImageViewer != FNullImageViewer::GetNullImageViewer())
 	{
 		Library->OnImageViewerOpened(InImageViewer);
 	}
 
+	if (InPosition == EMediaImageViewerPosition::Second)
+	{
+		RequestedView = EMediaImageViewerActivePosition::Both;
+	}
+
 	InvalidateView();
 }
 
 void SMediaViewer::ClearImageViewer(EMediaImageViewerPosition InPosition)
 {
-	ImageViewers[static_cast<int32>(InPosition)] = FNullImageViewer::GetNullImageViewer();
-
-	InvalidateView();
+	SetImageViewer(InPosition, FNullImageViewer::GetNullImageViewer());
 }
 
 void SMediaViewer::BindCommands()
@@ -176,19 +187,43 @@ void SMediaViewer::BindCommands()
 	CommandList->MapAction(
 		Commands.ToggleLockedTransform,
 		FExecuteAction::CreateSP(this, &SMediaViewer::ToggleLockedTransform),
-		FCanExecuteAction(),
+		FCanExecuteAction::CreateSP(this, &SMediaViewer::IsShowingBothImageViewers),
 		FGetActionCheckState::CreateSP(this, &SMediaViewer::AreLockedTransformEnabled)
 	);
 	
-	CommandList->MapAction(Commands.SecondImageOpacityMinus, FExecuteAction::CreateSP(this, &SMediaViewer::AdjustSecondImageOpacity, -10.f));
-	CommandList->MapAction(Commands.SecondImageOpacityPlus, FExecuteAction::CreateSP(this, &SMediaViewer::AdjustSecondImageOpacity, 10.f));
+	CommandList->MapAction(
+		Commands.SecondImageOpacityMinus,
+		FExecuteAction::CreateSP(this, &SMediaViewer::AdjustSecondImageOpacity, -10.f),
+		FCanExecuteAction::CreateSP(this, &SMediaViewer::IsShowingBothImageViewers)
+	);
+	
+	CommandList->MapAction(
+		Commands.SecondImageOpacityPlus, 
+		FExecuteAction::CreateSP(this, &SMediaViewer::AdjustSecondImageOpacity, 10.f),
+		FCanExecuteAction::CreateSP(this, &SMediaViewer::IsShowingBothImageViewers)
+	);
+
+	CommandList->MapAction(
+		Commands.SwapAB, 
+		FExecuteAction::CreateSP(this, &SMediaViewer::SwapABImageViewers),
+		FCanExecuteAction::CreateSP(this, &SMediaViewer::IsShowingBothImageViewers)
+	);
 
-	CommandList->MapAction(Commands.SwapAB, FExecuteAction::CreateSP(this, &SMediaViewer::SwapABImageViewers));
 	CommandList->MapAction(Commands.ResetAllTransforms, FExecuteAction::CreateSP(this, &SMediaViewer::ResetTransformToAll));
 }
 
+bool SMediaViewer::IsShowingBothImageViewers() const
+{
+	return ActiveView == EMediaImageViewerActivePosition::Both;
+}
+
 void SMediaViewer::SwapABImageViewers()
 {
+	if (!IsShowingBothImageViewers())
+	{
+		return;
+	}
+
 	TSharedPtr<FMediaImageViewer> OldSecond = ImageViewers[static_cast<int32>(EMediaImageViewerPosition::Second)];	
 	ImageViewers[static_cast<int32>(EMediaImageViewerPosition::Second)] = ImageViewers[static_cast<int32>(EMediaImageViewerPosition::First)];
 	ImageViewers[static_cast<int32>(EMediaImageViewerPosition::First)] = OldSecond;
@@ -197,41 +232,32 @@ void SMediaViewer::SwapABImageViewers()
 	ImageViewers[static_cast<int32>(EMediaImageViewerPosition::First)]->GetPaintSettings().Offset = ImageViewers[static_cast<int32>(EMediaImageViewerPosition::Second)]->GetPaintSettings().Offset;
 	ImageViewers[static_cast<int32>(EMediaImageViewerPosition::Second)]->GetPaintSettings().Offset = OldOffset;
 
-	CreateView();
+	InvalidateView();
 }
 
 void SMediaViewer::SetSingleView()
 {
-	// No need to do anything
-	if (IsImageViewerNull(EMediaImageViewerPosition::Second))
+	if (ActiveView == EMediaImageViewerActivePosition::Single)
 	{
 		return;
 	}
 
-	// The second entry is valid, so just use that.
-	if (IsImageViewerNull(EMediaImageViewerPosition::First))
-	{
-		ImageViewers[static_cast<int32>(EMediaImageViewerPosition::First)] = ImageViewers[static_cast<int32>(EMediaImageViewerPosition::Second)];
-	}
-
-	ImageViewers[static_cast<int32>(EMediaImageViewerPosition::Second)] = FNullImageViewer::GetNullImageViewer();
+	RequestedView = EMediaImageViewerActivePosition::Single;
 
-	CreateView();
+	InvalidateView();
 }
 
 void SMediaViewer::SetABOrientation(EOrientation InOrientation)
 {
-	if (MediaViewerSettings.ABOrientation == InOrientation)
+	if (MediaViewerSettings.ABOrientation == InOrientation
+		&& ActiveView == EMediaImageViewerActivePosition::Both)
 	{
 		return;
 	}
 
 	MediaViewerSettings.ABOrientation = InOrientation;
-
-	if (ActiveViews == EMediaImageViewerActivePosition::Both)
-	{
-		CreateView();
-	}
+	RequestedView = EMediaImageViewerActivePosition::Both;
+	InvalidateView();
 }
 
 void SMediaViewer::Tick(const FGeometry& AllottedGeometry, const double InCurrentTime, const float InDeltaTime)
@@ -395,29 +421,18 @@ int32 SMediaViewer::OnPaint(const FPaintArgs& InArgs, const FGeometry& InAllotte
 	const float ABSplitterLocation = MediaViewerSettings.ABSplitterLocation / 100.f;
 	const float SecondImageOpacity = MediaViewerSettings.SecondImageOpacity / 100.f;
 
-	switch (ActiveViews)
+	switch (ActiveView)
 	{
-		case EMediaImageViewerActivePosition::None:
-			break;
-
-		case EMediaImageViewerActivePosition::First:
-			ImageViewers[static_cast<int32>(EMediaImageViewerPosition::First)]->Paint(PaintParams);
-			break;
-
-		case EMediaImageViewerActivePosition::Second:
-			ImageViewers[static_cast<int32>(EMediaImageViewerPosition::Second)]->Paint(PaintParams);
+		case EMediaImageViewerActivePosition::Single:
+			GetImageViewer(EMediaImageViewerPosition::First)->Paint(PaintParams);
 			break;
 
 		case EMediaImageViewerActivePosition::Both:
 			if (!FMath::IsNearlyZero(ABSplitterLocation))
 			{
-				if (FMath::IsNearlyEqual(SecondImageOpacity, 1.f))
-				{
-					UVRange.SetLowerBound(0.f);
-					UVRange.SetUpperBound(ABSplitterLocation);
-				}
-
-				ImageViewers[static_cast<int32>(EMediaImageViewerPosition::First)]->Paint(PaintParams);
+				UVRange.SetLowerBound(0.f);
+				UVRange.SetUpperBound(1.f);
+				GetImageViewer(EMediaImageViewerPosition::First)->Paint(PaintParams);
 			}
 
 			if (!FMath::IsNearlyZero(ABSplitterLocation - 1.f))
@@ -425,7 +440,7 @@ int32 SMediaViewer::OnPaint(const FPaintArgs& InArgs, const FGeometry& InAllotte
 				UVRange.SetLowerBound(ABSplitterLocation);
 				UVRange.SetUpperBound(1.f);
 				PaintParams.ImageOpacity = SecondImageOpacity;
-				ImageViewers[static_cast<int32>(EMediaImageViewerPosition::Second)]->Paint(PaintParams);
+				GetImageViewer(EMediaImageViewerPosition::Second)->Paint(PaintParams);
 			}
 			break;
 	}
@@ -478,7 +493,7 @@ void SMediaViewer::CreateDelegates()
 	Delegates->SetSingleView.BindSP(this, &SMediaViewer::SetSingleView);
 	Delegates->GetABOrientation.BindSPLambda(this, [this] { return MediaViewerSettings.ABOrientation; });
 	Delegates->SetABOrientation.BindSP(this, &SMediaViewer::SetABOrientation);
-	Delegates->GetActiveViews.BindSPLambda(this, [this] { return ActiveViews; });
+	Delegates->GetActiveView.BindSPLambda(this, [this] { return ActiveView; });
 	Delegates->GetSettings.BindSPLambda(this, [this]{ return &MediaViewerSettings; });
 	Delegates->AreTransformsLocked.BindSPLambda(this, [this] { return MediaViewerSettings.bAreTransformsLocked; });
 	Delegates->ToggleLockedTransform.BindSPLambda(this, [this] { MediaViewerSettings.bAreTransformsLocked = !MediaViewerSettings.bAreTransformsLocked; });
@@ -510,10 +525,10 @@ void SMediaViewer::CreateDelegates()
 	Delegates->GetCommandListForPosition.BindSP(this, &SMediaViewer::GetOverlayCommandList);
 }
 
-TSharedRef<SMediaImageViewerOverlay> SMediaViewer::CreateOverlay(EMediaImageViewerPosition InPosition, bool bInCreateSplitDropTarget)
+TSharedRef<SMediaImageViewerOverlay> SMediaViewer::CreateOverlay(EMediaImageViewerPosition InPosition, bool bInShowPanelName)
 {
 	return  SNew(SMediaImageViewerOverlay, InPosition, Delegates.ToSharedRef())
-		.bCreateSplitDropTarget(bInCreateSplitDropTarget);
+		.bShowPanelName(bInShowPanelName);
 }
 
 TSharedRef<SWidget> SMediaViewer::CreateToolbar()
@@ -532,32 +547,14 @@ TSharedRef<SWidget> SMediaViewer::CreateToolbar()
 
 void SMediaViewer::CreateView()
 {
-	ActiveViews = EMediaImageViewerActivePosition::None;
-
-	if (!IsImageViewerNull(EMediaImageViewerPosition::First))
-	{
-		ActiveViews |= EMediaImageViewerActivePosition::First;
-	}
+	ActiveView = RequestedView;
 
-	if (!IsImageViewerNull(EMediaImageViewerPosition::Second))
+	switch (ActiveView)
 	{
-		ActiveViews |= EMediaImageViewerActivePosition::Second;
-	}
-
-	switch (ActiveViews)
-	{
-		case EMediaImageViewerActivePosition::None:
-			CreateEmptyView();
-			break;
-
-		case EMediaImageViewerActivePosition::First:
+		case EMediaImageViewerActivePosition::Single:
 			CreateSingleView(EMediaImageViewerPosition::First);
 			break;
 
-		case EMediaImageViewerActivePosition::Second:
-			CreateSingleView(EMediaImageViewerPosition::Second);
-			break;
-
 		case EMediaImageViewerActivePosition::Both:
 			CreateABView();
 			break;
@@ -571,18 +568,10 @@ void SMediaViewer::CreateView()
 	bInvalidated = false;
 }
 
-void SMediaViewer::CreateEmptyView()
-{
-	// Attach a null widget
-	ContentSlot->AttachWidget(
-		SNew(SMediaViewerDropTarget, Delegates.ToSharedRef())
-	);
-}
-
 void SMediaViewer::CreateSingleView(EMediaImageViewerPosition InPosition)
 {
 	// Attach the single overlay
-	ContentSlot->AttachWidget(CreateOverlay(InPosition, /* Create Split Drop Target */ true));
+	ContentSlot->AttachWidget(CreateOverlay(InPosition, /* Show Panel Name */ false));
 }
 
 void SMediaViewer::CreateABView()
@@ -602,7 +591,7 @@ void SMediaViewer::CreateABView()
 		.SizeRule(SSplitter::FractionOfParent)
 		.OnSlotResized(this, &SMediaViewer::OnABResized)
 		[
-			CreateOverlay(EMediaImageViewerPosition::First, /* Create Split Drop Target */ false)
+			CreateOverlay(EMediaImageViewerPosition::First, /* Show Panel Name */ true)
 		]
 
 		+ SSplitter::Slot()
@@ -611,33 +600,35 @@ void SMediaViewer::CreateABView()
 		.Resizable(true)
 		.SizeRule(SSplitter::FractionOfParent)
 		[
-			CreateOverlay(EMediaImageViewerPosition::Second, /* Create Split Drop Target */ false)
+			CreateOverlay(EMediaImageViewerPosition::Second, /* Show Panel Name */ true)
 		]
 	);	
 }
 
 TSharedPtr<SMediaImageViewerOverlay> SMediaViewer::GetOverlay(EMediaImageViewerPosition InPosition) const
 {
-	switch (ActiveViews)
+	TSharedRef<SWidget> ContentWidget = ContentSlot->GetWidget();
+
+	switch (ActiveView)
 	{
-		case EMediaImageViewerActivePosition::First:
+		case EMediaImageViewerActivePosition::Single:
 			if (InPosition == EMediaImageViewerPosition::First)
 			{
-				return StaticCastSharedRef<SMediaImageViewerOverlay>(ContentSlot->GetWidget());
+				return StaticCastSharedRef<SMediaImageViewerOverlay>(ContentWidget);
 			}
 			break;
 
-		case EMediaImageViewerActivePosition::Second:
-			if (InPosition == EMediaImageViewerPosition::Second)
+		case EMediaImageViewerActivePosition::Both:
+		{
+			TSharedRef<SSplitter> Splitter = StaticCastSharedRef<SSplitter>(ContentSlot->GetWidget());
+
+			if (Splitter->NumSlots() == static_cast<int32>(EMediaImageViewerPosition::COUNT))
 			{
-				return StaticCastSharedRef<SMediaImageViewerOverlay>(ContentSlot->GetWidget());
+				return StaticCastSharedRef<SMediaImageViewerOverlay>(
+					Splitter->SlotAt(static_cast<int32>(InPosition)).GetWidget()
+				);
 			}
-			break;
-
-		case EMediaImageViewerActivePosition::Both:
-			return StaticCastSharedRef<SMediaImageViewerOverlay>(
-				StaticCastSharedRef<SSplitter>(ContentSlot->GetWidget())->SlotAt(static_cast<int32>(InPosition)).GetWidget()
-			);
+		}
 	}
 
 	return nullptr;
@@ -694,7 +685,8 @@ void SMediaViewer::InvalidateView()
 
 bool SMediaViewer::IsImageViewerNull(EMediaImageViewerPosition InPosition) const
 {
-	return ImageViewers[static_cast<int32>(InPosition)] == FNullImageViewer::GetNullImageViewer();
+	return !ImageViewers[static_cast<int32>(InPosition)].IsValid()
+		|| ImageViewers[static_cast<int32>(InPosition)]->GetInfo().Id == FNullImageViewer::GetNullImageViewer()->GetInfo().Id;
 }
 
 void SMediaViewer::AddOffsetToAll(const FVector& InOffset)
diff --git a/Engine/Plugins/Experimental/MediaViewer/Source/MediaViewer/Private/Widgets/SMediaViewer.h b/Engine/Plugins/Experimental/MediaViewer/Source/MediaViewer/Private/Widgets/SMediaViewer.h
--- a/Engine/Plugins/Experimental/MediaViewer/Source/MediaViewer/Private/Widgets/SMediaViewer.h
+++ b/Engine/Plugins/Experimental/MediaViewer/Source/MediaViewer/Private/Widgets/SMediaViewer.h
@@ -36,12 +36,9 @@ namespace UE::MediaViewer::Private
 
 enum class EMediaImageViewerActivePosition : uint8
 {
-	None = 0,
-	First = 1 << 0,
-	Second = 1 << 1,
-	Both = First + Second
+	Single,
+	Both
 };
-ENUM_CLASS_FLAGS(EMediaImageViewerActivePosition)
 
 class SMediaViewer : public SCompoundWidget, public FGCObject
 {
@@ -70,6 +67,8 @@ public:
 	void SetImageViewer(EMediaImageViewerPosition InPosition, const TSharedRef<FMediaImageViewer>& InImageViewer);
 	
 	void ClearImageViewer(EMediaImageViewerPosition InPosition);
+	
+	bool IsShowingBothImageViewers() const;
 
 	void SwapABImageViewers();
 
@@ -126,8 +125,10 @@ protected:
 
 	/** Viewers currently open in the viewer. May be FNullImageViewer for an invalid viewer. */
 	TSharedPtr<FMediaImageViewer> ImageViewers[static_cast<int32>(EMediaImageViewerPosition::COUNT)];
-	/** Bitflag of the active views. */
-	EMediaImageViewerActivePosition ActiveViews;
+	/** The active view configuration. */
+	EMediaImageViewerActivePosition ActiveView;
+	/** The requested view configuration. */
+	EMediaImageViewerActivePosition RequestedView;
 
 	FIntPoint CursorLocalPosition;
 
@@ -144,7 +145,7 @@ protected:
 	void BindCommands();
 
 	/** Create and overlay widget for the given position. */
-	TSharedRef<SMediaImageViewerOverlay> CreateOverlay(EMediaImageViewerPosition InPosition, bool bInCreateSplitDropTarget);
+	TSharedRef<SMediaImageViewerOverlay> CreateOverlay(EMediaImageViewerPosition InPosition, bool bInShowPanelName);
 
 	/** Creates the toolbar. */
 	TSharedRef<SWidget> CreateToolbar();
@@ -152,9 +153,6 @@ protected:
 	/** Examines the available images and creates an appropriate view. */
 	void CreateView();
 
-	/** Creates a view for no image viewers. */
-	void CreateEmptyView();
-
 	/** Creates a view for a single image viewer. */
 	void CreateSingleView(EMediaImageViewerPosition InPosition);
 
diff --git a/Engine/Plugins/Experimental/MediaViewer/Source/MediaViewer/Private/Widgets/SMediaViewerDropTarget.cpp b/Engine/Plugins/Experimental/MediaViewer/Source/MediaViewer/Private/Widgets/SMediaViewerDropTarget.cpp
--- a/Engine/Plugins/Experimental/MediaViewer/Source/MediaViewer/Private/Widgets/SMediaViewerDropTarget.cpp
+++ b/Engine/Plugins/Experimental/MediaViewer/Source/MediaViewer/Private/Widgets/SMediaViewerDropTarget.cpp
@@ -10,6 +10,7 @@
 #include "Library/MediaViewerLibraryItem.h"
 #include "Styling/AppStyle.h"
 #include "Styling/StyleColors.h"
+#include "Widgets/Layout/SBox.h"
 #include "Widgets/MediaViewerDelegates.h"
 #include "Widgets/MediaViewerLibraryItemDragDropOperation.h"
 #include "Widgets/SBoxPanel.h"
@@ -29,86 +30,6 @@ void SMediaViewerDropTarget::Construct(const FArguments& InArgs, const TSharedRe
 	Delegates = InDelegates;
 	Position = InArgs._Position;
 
-	const EMediaImageViewerActivePosition ActiveViews = Delegates->GetActiveViews.Execute();
-	
-	const FText DropTargetMessageLower = LOCTEXT("DropTargetMessage", "Drop supported asset or library item here.");
-	FText DropTargetMessageUpper = FText::GetEmpty();
-
-	switch (ActiveViews)
-	{
-		default:
-		case EMediaImageViewerActivePosition::None:
-			break;
-
-		case EMediaImageViewerActivePosition::First:
-		case EMediaImageViewerActivePosition::Second:
-			if (Position == EMediaImageViewerPosition::First)
-			{
-				DropTargetMessageUpper = LOCTEXT("DropTargetMessageReplaceImage", "Replace Image");
-			}
-			else if (Position == EMediaImageViewerPosition::Second)
-			{
-				DropTargetMessageUpper = LOCTEXT("DropTargetMessageAddSecond", "Add Second Image");
-			}
-			break;
-
-		case EMediaImageViewerActivePosition::Both:
-			if (Position == EMediaImageViewerPosition::First || Position == EMediaImageViewerPosition::Second)
-			{
-				DropTargetMessageUpper = LOCTEXT("DropTargetMessageReplaceImage", "Replace Image");
-			}
-			break;
-	}
-
-	TSharedRef<SVerticalBox> MessageBox = SNew(SVerticalBox);
-
-	if (DropTargetMessageUpper.IsEmpty())
-	{
-		MessageBox->AddSlot()
-			.FillHeight(1.f)
-			.HAlign(HAlign_Center)
-			.VAlign(VAlign_Center)
-			.Padding(5.f)
-			[
-				SNew(STextBlock)
-				.ColorAndOpacity(FStyleColors::Foreground.GetSpecifiedColor())
-				.ShadowColorAndOpacity(FStyleColors::Panel.GetSpecifiedColor())
-				.ShadowOffset(FVector2D(1.f, 1.f))
-				.Visibility(this, &SMediaViewerDropTarget::GetDragDescriptionVisibility)
-				.Text(DropTargetMessageLower)
-			];
-	}
-	else
-	{
-		MessageBox->AddSlot()
-			.FillHeight(1.f)
-			.HAlign(HAlign_Center)
-			.VAlign(VAlign_Bottom)
-			.Padding(5.f)
-			[
-				SNew(STextBlock)
-				.ColorAndOpacity(FStyleColors::Foreground.GetSpecifiedColor())
-				.ShadowColorAndOpacity(FStyleColors::Panel.GetSpecifiedColor())
-				.ShadowOffset(FVector2D(1.f, 1.f))
-				.Visibility(this, &SMediaViewerDropTarget::GetDragDescriptionVisibility)
-				.Text(DropTargetMessageUpper)
-			];
-
-		MessageBox->AddSlot()
-			.FillHeight(1.f)
-			.HAlign(HAlign_Center)
-			.VAlign(VAlign_Top)
-			.Padding(5.f)
-			[
-				SNew(STextBlock)
-				.ColorAndOpacity(FStyleColors::Foreground.GetSpecifiedColor())
-				.ShadowColorAndOpacity(FStyleColors::Panel.GetSpecifiedColor())
-				.ShadowOffset(FVector2D(1.f, 1.f))
-				.Visibility(this, &SMediaViewerDropTarget::GetDragDescriptionVisibility)
-				.Text(DropTargetMessageLower)
-			];
-	}
-
 	SDropTarget::Construct(
 		SDropTarget::FArguments()
 		.OnDropped(this, &SMediaViewerDropTarget::OnDropped)
@@ -120,7 +41,17 @@ void SMediaViewerDropTarget::Construct(const FArguments& InArgs, const TSharedRe
 			]
 			+ SOverlay::Slot()
 			[
-				MessageBox
+				SNew(SBox)
+				.HAlign(HAlign_Center)
+				.VAlign(VAlign_Center)
+				[
+					SNew(STextBlock)
+					.ColorAndOpacity(FStyleColors::Foreground.GetSpecifiedColor())
+					.ShadowColorAndOpacity(FStyleColors::Panel.GetSpecifiedColor())
+					.ShadowOffset(FVector2D(1.f, 1.f))
+					.Visibility(this, &SMediaViewerDropTarget::GetDragDescriptionVisibility)
+					.Text(LOCTEXT("DropTargetMessage", "Drop supported asset or library item here."))
+				]
 			]
 		]);
 }
diff --git a/Engine/Plugins/Experimental/MediaViewer/Source/MediaViewer/Private/Widgets/SMediaViewerToolbar.cpp b/Engine/Plugins/Experimental/MediaViewer/Source/MediaViewer/Private/Widgets/SMediaViewerToolbar.cpp
--- a/Engine/Plugins/Experimental/MediaViewer/Source/MediaViewer/Private/Widgets/SMediaViewerToolbar.cpp
+++ b/Engine/Plugins/Experimental/MediaViewer/Source/MediaViewer/Private/Widgets/SMediaViewerToolbar.cpp
@@ -126,13 +126,13 @@ TSharedRef<SWidget> SMediaViewerToolbar::MakeCenterToolbar()
 	ToolbarBuilder.SetStyle(&FAppStyle::Get(), "ViewportToolbar");
 	ToolbarBuilder.SetIsFocusable(false);
 
-	EMediaImageViewerActivePosition ActiveViews = Delegates->GetActiveViews.Execute();
+	EMediaImageViewerActivePosition ActiveView = Delegates->GetActiveView.Execute();
 
 	const FMediaViewerCommands& Commands = FMediaViewerCommands::Get();
 
 	FSlateIcon ViewIcon;
 
-	if (Delegates->GetActiveViews.Execute() != EMediaImageViewerActivePosition::Both)
+	if (Delegates->GetActiveView.Execute() != EMediaImageViewerActivePosition::Both)
 	{
 		ViewIcon = FSlateIcon(FAppStyle::Get().GetStyleSetName(), "EditorViewport.ViewportConfig_OnePane");
 	}
@@ -149,13 +149,82 @@ TSharedRef<SWidget> SMediaViewerToolbar::MakeCenterToolbar()
 	{
 		ToolbarBuilder.BeginBlockGroup();
 		{
-			ToolbarBuilder.AddComboButton(
-				FUIAction(),
-				FOnGetContent::CreateSP(this, &SMediaViewerToolbar::MakeViewMenu),
-				FText::GetEmpty(),
-				LOCTEXT("ViewType", "Change the type of view between single, AB horizontal and AB vertical"),
-				ViewIcon
+			FButtonArgs SingleViewArgs;
+			SingleViewArgs.ToolTipOverride = LOCTEXT("SetSingleViewToolTip", "View a single image.");
+			SingleViewArgs.IconOverride = FSlateIcon(FAppStyle::Get().GetStyleSetName(), "EditorViewport.ViewportConfig_OnePane");
+			SingleViewArgs.UserInterfaceActionType = EUserInterfaceActionType::RadioButton;
+
+			SingleViewArgs.Action.ExecuteAction.BindSPLambda(
+				this,
+				[this]
+				{
+					Delegates->SetSingleView.Execute();
+				}
+			);
+
+			SingleViewArgs.Action.GetActionCheckState.BindSPLambda(
+				this,
+				[this]
+				{
+					return Delegates->GetActiveView.Execute() == EMediaImageViewerActivePosition::Single
+						? ECheckBoxState::Checked
+						: ECheckBoxState::Unchecked;
+				}
+			);
+
+			ToolbarBuilder.AddToolBarButton(SingleViewArgs);
+
+			FButtonArgs ABHorizontalViewArgs;
+			ABHorizontalViewArgs.ToolTipOverride = LOCTEXT("SetABHorizontalViewToolTip", "View 2 images side by side.");
+			ABHorizontalViewArgs.IconOverride = FSlateIcon(FAppStyle::Get().GetStyleSetName(), "EditorViewport.ViewportConfig_TwoPanesH");
+			ABHorizontalViewArgs.UserInterfaceActionType = EUserInterfaceActionType::RadioButton;
+
+			ABHorizontalViewArgs.Action.ExecuteAction.BindSPLambda(
+				this,
+				[this]
+				{
+					Delegates->SetABOrientation.Execute(Orient_Horizontal);
+				}
+			);
+
+			ABHorizontalViewArgs.Action.GetActionCheckState.BindSPLambda(
+				this,
+				[this]
+				{
+					return (Delegates->GetActiveView.Execute() == EMediaImageViewerActivePosition::Both
+						&& Delegates->GetABOrientation.Execute() == Orient_Horizontal)
+						? ECheckBoxState::Checked
+						: ECheckBoxState::Unchecked;
+				}
+			);
+
+			ToolbarBuilder.AddToolBarButton(ABHorizontalViewArgs);
+
+			FButtonArgs ABVerticalViewArgs;
+			ABVerticalViewArgs.ToolTipOverride = LOCTEXT("SetABVerticalViewToolTip", "View 2 images, one above the other.");
+			ABVerticalViewArgs.IconOverride = FSlateIcon(FAppStyle::Get().GetStyleSetName(), "EditorViewport.ViewportConfig_TwoPanesV");
+			ABVerticalViewArgs.UserInterfaceActionType = EUserInterfaceActionType::RadioButton;
+
+			ABVerticalViewArgs.Action.ExecuteAction.BindSPLambda(
+				this,
+				[this]
+				{
+					Delegates->SetABOrientation.Execute(Orient_Vertical);
+				}
 			);
+
+			ABVerticalViewArgs.Action.GetActionCheckState.BindSPLambda(
+				this,
+				[this]
+				{
+					return (Delegates->GetActiveView.Execute() == EMediaImageViewerActivePosition::Both
+						&& Delegates->GetABOrientation.Execute() == Orient_Vertical)
+						? ECheckBoxState::Checked
+						: ECheckBoxState::Unchecked;
+				}
+			);
+
+			ToolbarBuilder.AddToolBarButton(ABVerticalViewArgs);
 		}
 		ToolbarBuilder.EndBlockGroup();
 
@@ -316,109 +385,6 @@ TSharedRef<SWidget> SMediaViewerToolbar::GetBackgroundTextureSettingsWidget() co
 	return MediaViewerSettingsWidget.ToSharedRef();
 }
 
-TSharedRef<SWidget> SMediaViewerToolbar::MakeViewMenu()
-{
-	FMenuBuilder MenuBuilder(true, nullptr, nullptr, false, &FAppStyle::Get(), false);
-
-	FMenuEntryParams SingleViewParams;
-	SingleViewParams.LabelOverride = LOCTEXT("SetSingleView", "Single Image");
-	SingleViewParams.ToolTipOverride = LOCTEXT("SetSingleViewToolTip", "View a single image.");
-	SingleViewParams.IconOverride = FSlateIcon(FAppStyle::Get().GetStyleSetName(), "EditorViewport.ViewportConfig_OnePane");
-	SingleViewParams.UserInterfaceActionType = EUserInterfaceActionType::RadioButton;
-
-	SingleViewParams.DirectActions.ExecuteAction.BindSPLambda(
-		this,
-		[this]
-		{
-			Delegates->SetSingleView.Execute();
-		}
-	);
-
-	SingleViewParams.DirectActions.GetActionCheckState.BindSPLambda(
-		this,
-		[this]
-		{
-			return Delegates->GetActiveViews.Execute() != EMediaImageViewerActivePosition::Both
-				? ECheckBoxState::Checked
-				: ECheckBoxState::Unchecked;
-		}
-	);
-
-	MenuBuilder.AddMenuEntry(SingleViewParams);
-
-	FMenuEntryParams ABHorizontalViewParams;
-	ABHorizontalViewParams.LabelOverride = LOCTEXT("SetABHorizontalView", "Horizontal");
-	ABHorizontalViewParams.ToolTipOverride = LOCTEXT("SetABHorizontalViewToolTip", "View 2 images side by side.");
-	ABHorizontalViewParams.IconOverride = FSlateIcon(FAppStyle::Get().GetStyleSetName(), "EditorViewport.ViewportConfig_TwoPanesH");
-	ABHorizontalViewParams.UserInterfaceActionType = EUserInterfaceActionType::RadioButton;
-
-	ABHorizontalViewParams.DirectActions.ExecuteAction.BindSPLambda(
-		this,
-		[this]
-		{
-			Delegates->SetABOrientation.Execute(Orient_Horizontal);
-		}
-	);
-
-	ABHorizontalViewParams.DirectActions.CanExecuteAction.BindSPLambda(
-		this,
-		[this]
-		{
-			return Delegates->GetActiveViews.Execute() == EMediaImageViewerActivePosition::Both;
-		}
-	);
-
-	ABHorizontalViewParams.DirectActions.GetActionCheckState.BindSPLambda(
-		this,
-		[this]
-		{
-			return (Delegates->GetActiveViews.Execute() == EMediaImageViewerActivePosition::Both
-				&& Delegates->GetABOrientation.Execute() == Orient_Horizontal)
-				? ECheckBoxState::Checked
-				: ECheckBoxState::Unchecked;
-		}
-	);
-
-	MenuBuilder.AddMenuEntry(ABHorizontalViewParams);
-
-	FMenuEntryParams ABVerticalViewParams;
-	ABVerticalViewParams.LabelOverride = LOCTEXT("SetABVerticalView", "Vertical");
-	ABVerticalViewParams.ToolTipOverride = LOCTEXT("SetABVerticalViewToolTip", "View 2 images, one above the other.");
-	ABVerticalViewParams.IconOverride = FSlateIcon(FAppStyle::Get().GetStyleSetName(), "EditorViewport.ViewportConfig_TwoPanesV");
-	ABVerticalViewParams.UserInterfaceActionType = EUserInterfaceActionType::RadioButton;
-
-	ABVerticalViewParams.DirectActions.ExecuteAction.BindSPLambda(
-		this,
-		[this]
-		{
-			Delegates->SetABOrientation.Execute(Orient_Vertical);
-		}
-	);
-
-	ABVerticalViewParams.DirectActions.CanExecuteAction.BindSPLambda(
-		this,
-		[this]
-		{
-			return Delegates->GetActiveViews.Execute() == EMediaImageViewerActivePosition::Both;
-		}
-	);
-
-	ABVerticalViewParams.DirectActions.GetActionCheckState.BindSPLambda(
-		this,
-		[this]
-		{
-			return (Delegates->GetActiveViews.Execute() == EMediaImageViewerActivePosition::Both
-				&& Delegates->GetABOrientation.Execute() == Orient_Vertical)
-				? ECheckBoxState::Checked
-				: ECheckBoxState::Unchecked;
-		}
-	);
-
-	MenuBuilder.AddMenuEntry(ABVerticalViewParams);
-
-	return MenuBuilder.MakeWidget();
-}
-
 } // UE::MediaViewer::Private
 
 #undef LOCTEXT_NAMESPACE
diff --git a/Engine/Plugins/Experimental/MediaViewer/Source/MediaViewer/Private/Widgets/SMediaViewerToolbar.h b/Engine/Plugins/Experimental/MediaViewer/Source/MediaViewer/Private/Widgets/SMediaViewerToolbar.h
--- a/Engine/Plugins/Experimental/MediaViewer/Source/MediaViewer/Private/Widgets/SMediaViewerToolbar.h
+++ b/Engine/Plugins/Experimental/MediaViewer/Source/MediaViewer/Private/Widgets/SMediaViewerToolbar.h
@@ -53,8 +53,6 @@ protected:
 
 	TSharedRef<SWidget> GetBackgroundTextureSettingsWidget() const;
 
-	TSharedRef<SWidget> MakeViewMenu();
-
 	FText GetScaleMenuLabel(EMediaImageViewerPosition InPosition) const;
 	TSharedRef<SWidget> MakeScaleMenu(EMediaImageViewerPosition InPosition) const;
 	
