Media Viewer: Content browser integration

diff --git a/Engine/Plugins/Experimental/MediaViewer/Source/MediaViewer/Private/MediaViewerContentBrowserIntegration.cpp b/Engine/Plugins/Experimental/MediaViewer/Source/MediaViewer/Private/MediaViewerContentBrowserIntegration.cpp
new file mode 100644
--- /dev/null
+++ b/Engine/Plugins/Experimental/MediaViewer/Source/MediaViewer/Private/MediaViewerContentBrowserIntegration.cpp
@@ -0,0 +1,130 @@
+// Copyright Epic Games, Inc. All Rights Reserved.
+
+#include "MediaViewerContentBrowserIntegration.h"
+
+#include "AssetRegistry/AssetData.h"
+#include "ContentBrowserDelegates.h"
+#include "ContentBrowserModule.h"
+#include "Framework/Commands/UIAction.h"
+#include "Framework/MultiBox/MultiBoxBuilder.h"
+#include "Internationalization/Text.h"
+#include "MediaViewer.h"
+#include "MediaViewerModule.h"
+#include "Modules/ModuleManager.h"
+#include "Styling/AppStyle.h"
+#include "Textures/SlateIcon.h"
+#include "UObject/Class.h"
+
+#define LOCTEXT_NAMESPACE "MediaViewerContentBrowserIntegration"
+
+namespace UE::MediaViewer
+{
+
+void FMediaViewerContentBrowserIntegration::Integrate()
+{
+	Disintegrate();
+
+	FContentBrowserModule& ContentBrowserModule = FModuleManager::LoadModuleChecked<FContentBrowserModule>("ContentBrowser");
+	TArray<FContentBrowserMenuExtender_SelectedAssets>& CBMenuExtenderDelegates = ContentBrowserModule.GetAllAssetViewContextMenuExtenders();
+	CBMenuExtenderDelegates.Add(FContentBrowserMenuExtender_SelectedAssets::CreateSP(this, &FMediaViewerContentBrowserIntegration::OnExtendContentBrowserAssetSelectionMenu));
+	ContentBrowserAssetHandle = CBMenuExtenderDelegates.Last().GetHandle();
+}
+
+void FMediaViewerContentBrowserIntegration::Disintegrate()
+{
+	if (ContentBrowserAssetHandle.IsValid())
+	{
+		if (FContentBrowserModule* ContentBrowserModule = FModuleManager::GetModulePtr<FContentBrowserModule>("ContentBrowser"))
+		{
+			TArray<FContentBrowserMenuExtender_SelectedAssets>& CBMenuExtenderDelegates = ContentBrowserModule->GetAllAssetViewContextMenuExtenders();
+
+			CBMenuExtenderDelegates.RemoveAll(
+				[this](const FContentBrowserMenuExtender_SelectedAssets& InElement)
+				{
+					return InElement.GetHandle() == ContentBrowserAssetHandle;
+				});
+
+			ContentBrowserAssetHandle.Reset();
+		}
+	}
+}
+
+TSharedRef<FExtender> FMediaViewerContentBrowserIntegration::OnExtendContentBrowserAssetSelectionMenu(const TArray<FAssetData>& InSelectedAssets)
+{
+	using namespace UE::MediaViewer;
+
+	const IMediaViewerModule& MediaViewerModule = IMediaViewerModule::Get();
+
+	TSharedRef<FExtender> Extender = MakeShared<FExtender>();
+	int32 AssetCount = 0;
+
+	for (const FAssetData& SelectedAsset : InSelectedAssets)
+	{
+		if (MediaViewerModule.HasFactoryFor(SelectedAsset))
+		{
+			++AssetCount;
+		}
+	}
+
+	if (AssetCount == 0)
+	{
+		return Extender;
+	}
+
+	Extender->AddMenuExtension(
+		"GetAssetActions",
+		EExtensionHook::After,
+		nullptr,
+		FMenuExtensionDelegate::CreateLambda(
+			[InSelectedAssets, AssetCount](FMenuBuilder& InMenuBuilder)
+			{
+				const FText ToolTip = AssetCount == 1
+					? LOCTEXT("OpenInMediaViewerSingleTooltip", "Open this asset in the Media Viewer, replacing the Single View or A image.")
+					: LOCTEXT("OpenInMediaViewerMultiTooltip", "Open the first 2 valid assets in the Media Viewer in the A/B view.");
+
+				InMenuBuilder.AddMenuEntry(
+					LOCTEXT("OpenInMediaViewer", "Open in Media Viewer"),
+					ToolTip,
+					FSlateIcon(FAppStyle::GetAppStyleSetName(), "Sequencer.Tracks.Media"),
+					FUIAction(FExecuteAction::CreateStatic(&FMediaViewerContentBrowserIntegration::OpenInMediaViewer, InSelectedAssets))
+				);
+			}
+		)
+	);
+
+	return Extender;
+}
+
+void FMediaViewerContentBrowserIntegration::OpenInMediaViewer(TArray<FAssetData> InSelectedAssets)
+{
+	IMediaViewerModule& MediaViewerModule = IMediaViewerModule::Get();
+
+	if (!MediaViewerModule.OpenTab())
+	{
+		UE_LOG(LogMediaViewer, Error, TEXT("Unable to open Media Viewer Tab."));
+		return;
+	}
+
+	bool bHasOpenedAsset = false;
+
+	for (const FAssetData& SelectedAsset : InSelectedAssets)
+	{
+		if (MediaViewerModule.HasFactoryFor(SelectedAsset))
+		{
+			if (!bHasOpenedAsset)
+			{
+				MediaViewerModule.SetImage(EMediaImageViewerPosition::First, SelectedAsset);
+				bHasOpenedAsset = true;
+			}
+			else
+			{
+				MediaViewerModule.SetImage(EMediaImageViewerPosition::Second, SelectedAsset);
+				return;
+			}
+		}
+	}
+}
+
+} // UE::MediaViewer
+
+#undef LOCTEXT_NAMESPACE
diff --git a/Engine/Plugins/Experimental/MediaViewer/Source/MediaViewer/Private/MediaViewerContentBrowserIntegration.h b/Engine/Plugins/Experimental/MediaViewer/Source/MediaViewer/Private/MediaViewerContentBrowserIntegration.h
new file mode 100644
--- /dev/null
+++ b/Engine/Plugins/Experimental/MediaViewer/Source/MediaViewer/Private/MediaViewerContentBrowserIntegration.h
@@ -0,0 +1,37 @@
+// Copyright Epic Games, Inc. All Rights Reserved.
+
+#pragma once
+
+#include "Templates/SharedPointer.h"
+
+#include "Containers/ContainersFwd.h"
+#include "Delegates/IDelegateInstance.h"
+
+class FExtender;
+struct FAssetData;
+
+namespace UE::MediaViewer
+{
+
+class FMediaViewerContentBrowserIntegration : public TSharedFromThis<FMediaViewerContentBrowserIntegration>
+{
+public:
+	static const TSharedRef<FMediaViewerContentBrowserIntegration>& Get()
+	{
+		static TSharedRef<FMediaViewerContentBrowserIntegration> Integration = MakeShared<FMediaViewerContentBrowserIntegration>();
+		return Integration;
+	}
+
+	void Integrate();
+
+	void Disintegrate();
+
+protected:
+	static void OpenInMediaViewer(TArray<FAssetData> InSelectedAssets);
+
+	FDelegateHandle ContentBrowserAssetHandle;
+
+	TSharedRef<FExtender> OnExtendContentBrowserAssetSelectionMenu(const TArray<FAssetData>& InSelectedAssets);
+};
+
+} // UE::MediaViewer
diff --git a/Engine/Plugins/Experimental/MediaViewer/Source/MediaViewer/Private/MediaViewerModule.cpp b/Engine/Plugins/Experimental/MediaViewer/Source/MediaViewer/Private/MediaViewerModule.cpp
--- a/Engine/Plugins/Experimental/MediaViewer/Source/MediaViewer/Private/MediaViewerModule.cpp
+++ b/Engine/Plugins/Experimental/MediaViewer/Source/MediaViewer/Private/MediaViewerModule.cpp
@@ -15,6 +15,7 @@
 #include "Library/MediaViewerLibraryIni.h"
 #include "MediaViewer.h"
 #include "MediaViewerCommands.h"
+#include "MediaViewerContentBrowserIntegration.h"
 #include "MediaViewerStyle.h"
 #include "Misc/NotNull.h"
 #include "Modules/ModuleManager.h"
@@ -25,6 +26,8 @@
 #include "WorkspaceMenuStructure.h"
 #include "WorkspaceMenuStructureModule.h"
 
+DEFINE_LOG_CATEGORY(LogMediaViewer);
+
 #define LOCTEXT_NAMESPACE "MediaViewerModule"
 
 namespace UE::MediaViewer::Private
@@ -324,6 +327,7 @@ void FMediaViewerModule::StartupModule()
 	FMediaViewerCommands::Register();
 	FMediaViewerStyle::Get();
 	RegisterDefaultImageViewers();
+	FMediaViewerContentBrowserIntegration::Get()->Integrate();
 
 	FLevelEditorModule& LevelEditorModule = FModuleManager::GetModuleChecked<FLevelEditorModule>("LevelEditor");
 	LevelEditorModule.OnLevelEditorCreated().AddLambda(
@@ -342,6 +346,7 @@ void FMediaViewerModule::ShutdownModule()
 	UnregisterAllImageViewers();
 	UnregisterLevelEditorMenuItems();
 	FMediaViewerCommands::Unregister();
+	FMediaViewerContentBrowserIntegration::Get()->Disintegrate();
 }
 
 void FMediaViewerModule::OnTabClosed(TSharedRef<SDockTab> InDockTab)
diff --git a/Engine/Plugins/Experimental/MediaViewer/Source/MediaViewer/Public/MediaViewer.h b/Engine/Plugins/Experimental/MediaViewer/Source/MediaViewer/Public/MediaViewer.h
--- a/Engine/Plugins/Experimental/MediaViewer/Source/MediaViewer/Public/MediaViewer.h
+++ b/Engine/Plugins/Experimental/MediaViewer/Source/MediaViewer/Public/MediaViewer.h
@@ -3,8 +3,11 @@
 #pragma once
 
 #include "HAL/Platform.h"
+#include "Logging/LogMacros.h"
 #include "UObject/NameTypes.h"
 
+DECLARE_LOG_CATEGORY_EXTERN(LogMediaViewer, Log, All);
+
 namespace UE::MediaViewer
 {
 
