Media Viewer: DPI Fix
- DPI scaling issues should be fixed.
- Made overlay text sizes consistent.

diff --git a/Engine/Plugins/Experimental/MediaViewer/Source/MediaViewer/Private/ImageViewers/MediaImageViewer.cpp b/Engine/Plugins/Experimental/MediaViewer/Source/MediaViewer/Private/ImageViewers/MediaImageViewer.cpp
--- a/Engine/Plugins/Experimental/MediaViewer/Source/MediaViewer/Private/ImageViewers/MediaImageViewer.cpp
+++ b/Engine/Plugins/Experimental/MediaViewer/Source/MediaViewer/Private/ImageViewers/MediaImageViewer.cpp
@@ -69,9 +69,11 @@ FSlateClippingZone FMediaImageViewer::CreateSlateClippingZone(const FGeometry& I
 		const float Offset = InUVRange.GetLowerBound().GetValue() * InViewerSize.X;
 		const float Width = InUVRange.Size<float>() * InViewerSize.X;
 
+		Zone.TopLeft.X /= InGeometry.Scale;
 		Zone.TopLeft.X += Offset + InViewerPosition.X;
 		Zone.TopRight.X = Zone.TopLeft.X + Width;
 
+		Zone.BottomLeft.X /= InGeometry.Scale;
 		Zone.BottomLeft.X += Offset + InViewerPosition.X;
 		Zone.BottomRight.X = Zone.BottomLeft.X + Width;
 	}
@@ -80,9 +82,11 @@ FSlateClippingZone FMediaImageViewer::CreateSlateClippingZone(const FGeometry& I
 		const float Offset = InUVRange.GetLowerBound().GetValue() * InViewerSize.Y;
 		const float Height = InUVRange.Size<float>() * InViewerSize.Y;
 
+		Zone.TopLeft.Y /= InGeometry.Scale;
 		Zone.TopLeft.Y += Offset + InViewerPosition.Y;
 		Zone.BottomLeft.Y = Zone.TopLeft.Y + Height;
 
+		Zone.TopRight.Y /= InGeometry.Scale;
 		Zone.TopRight.Y += Offset + InViewerPosition.Y;
 		Zone.BottomRight.Y = Zone.TopRight.Y + Height;
 	}
diff --git a/Engine/Plugins/Experimental/MediaViewer/Source/MediaViewer/Private/MediaViewerStyle.cpp b/Engine/Plugins/Experimental/MediaViewer/Source/MediaViewer/Private/MediaViewerStyle.cpp
--- a/Engine/Plugins/Experimental/MediaViewer/Source/MediaViewer/Private/MediaViewerStyle.cpp
+++ b/Engine/Plugins/Experimental/MediaViewer/Source/MediaViewer/Private/MediaViewerStyle.cpp
@@ -2,10 +2,12 @@
 
 #include "MediaViewerStyle.h"
 
+#include "DetailLayoutBuilder.h"
 #include "Styling/AppStyle.h"
 #include "Styling/CoreStyle.h"
 #include "Styling/SlateStyleRegistry.h"
 #include "Styling/SlateTypes.h"
+#include "Styling/StyleColors.h"
 #include "Styling/ToolBarStyle.h"
 
 namespace UE::MediaViewer::Private
@@ -19,7 +21,11 @@ FMediaViewerStyle::FMediaViewerStyle()
 	FSlateStyleSet::SetContentRoot(FPaths::EngineContentDir() / TEXT("Editor/Slate"));
 	FSlateStyleSet::SetCoreContentRoot(FPaths::EngineContentDir() / TEXT("Slate"));
 
-	const FTextBlockStyle NormalText = FAppStyle::Get().GetWidgetStyle<FTextBlockStyle>("NormalText");
+	FTextBlockStyle NormalText = FAppStyle::Get().GetWidgetStyle<FTextBlockStyle>("NormalText");
+	NormalText.SetFont(IDetailLayoutBuilder::GetDetailFont());
+	NormalText.SetShadowColorAndOpacity(FStyleColors::Foreground.GetSpecifiedColor());
+	NormalText.SetShadowOffset(FVector2D(1.0, 1.0));
+	Set("RichTextBlock.Normal", NormalText);
 	Set("RichTextBlock.Red", FTextBlockStyle(NormalText).SetColorAndOpacity(FLinearColor(1.0f, 0.1f, 0.1f)));
 	Set("RichTextBlock.Green", FTextBlockStyle(NormalText).SetColorAndOpacity(FLinearColor(0.1f, 1.0f, 0.1f)));
 	Set("RichTextBlock.Blue", FTextBlockStyle(NormalText).SetColorAndOpacity(FLinearColor(0.1f, 0.1f, 1.0f)));
diff --git a/Engine/Plugins/Experimental/MediaViewer/Source/MediaViewer/Private/Widgets/MediaViewerDelegates.h b/Engine/Plugins/Experimental/MediaViewer/Source/MediaViewer/Private/Widgets/MediaViewerDelegates.h
--- a/Engine/Plugins/Experimental/MediaViewer/Source/MediaViewer/Private/Widgets/MediaViewerDelegates.h
+++ b/Engine/Plugins/Experimental/MediaViewer/Source/MediaViewer/Private/Widgets/MediaViewerDelegates.h
@@ -26,7 +26,7 @@ enum class EMediaImageViewerActivePosition : uint8;
 
 struct FMediaViewerDelegates
 {
-	DECLARE_DELEGATE_RetVal(FIntPoint, FGetLocation)
+	DECLARE_DELEGATE_RetVal(FVector2D, FGetLocation)
 	DECLARE_DELEGATE_RetVal(EOrientation, FGetOrientation)
 	DECLARE_DELEGATE_OneParam(FSetOrientation, EOrientation)
 	DECLARE_DELEGATE_RetVal(EMediaImageViewerActivePosition, FGetActiveView)
diff --git a/Engine/Plugins/Experimental/MediaViewer/Source/MediaViewer/Private/Widgets/SMediaImageViewerOverlay.cpp b/Engine/Plugins/Experimental/MediaViewer/Source/MediaViewer/Private/Widgets/SMediaImageViewerOverlay.cpp
--- a/Engine/Plugins/Experimental/MediaViewer/Source/MediaViewer/Private/Widgets/SMediaImageViewerOverlay.cpp
+++ b/Engine/Plugins/Experimental/MediaViewer/Source/MediaViewer/Private/Widgets/SMediaImageViewerOverlay.cpp
@@ -2,6 +2,7 @@
 
 #include "Widgets/SMediaImageViewerOverlay.h"
 
+#include "DetailLayoutBuilder.h"
 #include "EditorViewportClient.h"
 #include "Framework/Application/SlateApplication.h"
 #include "Framework/Commands/UICommandList.h"
@@ -103,7 +104,7 @@ const TSharedRef<FUICommandList>& SMediaImageViewerOverlay::GetCommandList() con
 
 FIntPoint SMediaImageViewerOverlay::GetImageViewerPixelCoordinates() const
 {
-	const FVector2f CursorLocation = GetImageViewerPixelCoordinates_Exact();
+	const FVector2D CursorLocation = GetImageViewerPixelCoordinates_Exact();
 
 	return FIntPoint(
 		FMath::RoundToInt(CursorLocation.X),
@@ -111,28 +112,27 @@ FIntPoint SMediaImageViewerOverlay::GetImageViewerPixelCoordinates() const
 	);
 }
 
-FVector2f SMediaImageViewerOverlay::GetImageViewerPixelCoordinates_Exact() const
+FVector2D SMediaImageViewerOverlay::GetImageViewerPixelCoordinates_Exact() const
 {
 	TSharedPtr<FMediaImageViewer> ImageViewer = Delegates->GetImageViewer.Execute(Position);
 
 	if (!ImageViewer.IsValid())
 	{
-		return FVector2f(-1, -1);
+		return FVector2D(-1, -1);
 	}
 
-	const FIntPoint ViewerPosition = Delegates->GetViewerPosition.Execute();
-	const FIntPoint ViewerSize = Delegates->GetViewerSize.Execute();
-	const FIntPoint CursorLocation = Delegates->GetCursorLocation.Execute() - ViewerPosition;
+	const FVector2D ViewerPosition = Delegates->GetViewerPosition.Execute();
+	const FVector2D ViewerSize = Delegates->GetViewerSize.Execute();
 	const FVector2D TopLeft = ImageViewer->GetPaintOffset(ViewerSize, Delegates->GetViewerPosition.Execute());
 
-	FVector2f CursorLocationF = FVector2f(CursorLocation.X, CursorLocation.Y);
-	CursorLocationF.X -= TopLeft.X;
-	CursorLocationF.Y -= TopLeft.Y;
-	CursorLocationF.X += ViewerPosition.X;
-	CursorLocationF.Y += ViewerPosition.Y;
-	CursorLocationF /= ImageViewer->GetPaintSettings().Scale;
+	FVector2D CursorLocation = Delegates->GetCursorLocation.Execute() - ViewerPosition;
+	CursorLocation.X -= TopLeft.X;
+	CursorLocation.Y -= TopLeft.Y;
+	CursorLocation.X += ViewerPosition.X;
+	CursorLocation.Y += ViewerPosition.Y;
+	CursorLocation /= ImageViewer->GetPaintSettings().Scale;
 
-	return CursorLocationF;
+	return CursorLocation;
 }
 
 bool SMediaImageViewerOverlay::IsCursorOverImageViewer() const
@@ -190,7 +190,7 @@ FReply SMediaImageViewerOverlay::OnMouseButtonDown(const FGeometry& InMyGeometry
 {
 	if (!bDragging && !bExternalDragging)
 	{
-		const FIntPoint CursorLocation = Delegates->GetCursorLocation.Execute();
+		const FVector2D CursorLocation = Delegates->GetCursorLocation.Execute();
 
 		if (InMouseEvent.GetEffectingButton() == EKeys::LeftMouseButton)
 		{
@@ -205,7 +205,7 @@ FReply SMediaImageViewerOverlay::OnMouseButtonDown(const FGeometry& InMyGeometry
 		{
 			const FInputEventState EventState(nullptr, InMouseEvent.GetEffectingButton(), IE_Pressed);
 
-			if (ImageViewer->OnTrackingStarted(EventState, CursorLocation))
+			if (ImageViewer->OnTrackingStarted(EventState, FIntPoint(FMath::RoundToInt(CursorLocation.X), FMath::RoundToInt(CursorLocation.Y))))
 			{
 				DragButtonName = InMouseEvent.GetEffectingButton().GetFName();
 				bExternalDragging = true;
@@ -278,7 +278,7 @@ bool SMediaImageViewerOverlay::SupportsKeyboardFocus() const
 	return true;
 }
 
-void SMediaImageViewerOverlay::UpdateMouse(const TOptional<FIntPoint>& InMousePosition)
+void SMediaImageViewerOverlay::UpdateMouse(const TOptional<FVector2D>& InMousePosition)
 {
 	if (!DragButtonName.IsNone())
 	{
@@ -319,7 +319,8 @@ void SMediaImageViewerOverlay::OnDragButtonUp(FName InKeyName)
 	{
 		if (TSharedPtr<FMediaImageViewer> ImageViewer = Delegates->GetImageViewer.Execute(Position))
 		{
-			ImageViewer->OnTrackingStopped(Delegates->GetCursorLocation.Execute());
+			const FVector2D CursorLocation = Delegates->GetCursorLocation.Execute();
+			ImageViewer->OnTrackingStopped(FIntPoint(FMath::RoundToInt(CursorLocation.X), FMath::RoundToInt(CursorLocation.Y)));
 		}
 
 		DragButtonName = NAME_None;
@@ -329,11 +330,11 @@ void SMediaImageViewerOverlay::OnDragButtonUp(FName InKeyName)
 
 void SMediaImageViewerOverlay::UpdateDragPosition()
 {
-	const FIntPoint CursorLocation = Delegates->GetCursorLocation.Execute();
+	const FVector2D CursorLocation = Delegates->GetCursorLocation.Execute();
 
 	FVector NewOffset = DragStartOffset;
-	NewOffset.X += (CursorLocation.X - DragStartCursor.X) /*/ ImageViewer->GetPaintSettings().Scale*/;
-	NewOffset.Y += (CursorLocation.Y - DragStartCursor.Y) /*/ ImageViewer->GetPaintSettings().Scale*/;
+	NewOffset.X += CursorLocation.X - DragStartCursor.X;
+	NewOffset.Y += CursorLocation.Y - DragStartCursor.Y;
 
 	Try_AddOffset(NewOffset - GetOffset());
 }
@@ -439,6 +440,7 @@ TSharedRef<SWidget> SMediaImageViewerOverlay::CreateOverlay(bool bInShowPanelNam
 				.Padding(5.f)
 				[
 					SNew(STextBlock)
+					.Font(IDetailLayoutBuilder::GetDetailFont())
 					.ColorAndOpacity(FStyleColors::Foreground.GetSpecifiedColor())
 					.ShadowColorAndOpacity(FStyleColors::Panel.GetSpecifiedColor())
 					.ShadowOffset(FVector2D(1.f, 1.f))
@@ -455,6 +457,7 @@ TSharedRef<SWidget> SMediaImageViewerOverlay::CreateOverlay(bool bInShowPanelNam
 			.Padding(5.f)
 			[
 				SNew(STextBlock)
+				.Font(IDetailLayoutBuilder::GetDetailFont())
 				.ColorAndOpacity(FStyleColors::Foreground.GetSpecifiedColor())
 				.ShadowColorAndOpacity(FStyleColors::Panel.GetSpecifiedColor())
 				.ShadowOffset(FVector2D(1.f, 1.f))
@@ -496,11 +499,11 @@ void SMediaImageViewerOverlay::ToggleOverlay()
 
 void SMediaImageViewerOverlay::MultiplyScaleAroundCursor(float InMultipler)
 {
-	const FVector2f CursorLocationBefore = GetImageViewerPixelCoordinates_Exact();
+	const FVector2D CursorLocationBefore = GetImageViewerPixelCoordinates_Exact();
 
 	SetScale(GetScale() * InMultipler);
 
-	const FVector2f CursorLocationAfter = GetImageViewerPixelCoordinates_Exact();
+	const FVector2D CursorLocationAfter = GetImageViewerPixelCoordinates_Exact();
 
 	if (!(CursorLocationBefore - CursorLocationAfter).IsNearlyZero())
 	{
diff --git a/Engine/Plugins/Experimental/MediaViewer/Source/MediaViewer/Private/Widgets/SMediaImageViewerOverlay.h b/Engine/Plugins/Experimental/MediaViewer/Source/MediaViewer/Private/Widgets/SMediaImageViewerOverlay.h
--- a/Engine/Plugins/Experimental/MediaViewer/Source/MediaViewer/Private/Widgets/SMediaImageViewerOverlay.h
+++ b/Engine/Plugins/Experimental/MediaViewer/Source/MediaViewer/Private/Widgets/SMediaImageViewerOverlay.h
@@ -48,12 +48,12 @@ public:
 	FIntPoint GetImageViewerPixelCoordinates() const;
 
 	/** Gets the pixel coordinate of mouse on the image. May be outside of the image. Subpixel accuracy */
-	FVector2f GetImageViewerPixelCoordinates_Exact() const;
+	FVector2D GetImageViewerPixelCoordinates_Exact() const;
 
 	bool IsCursorOverImageViewer() const;
 
 	/** Triggers a mouse update. */
-	void UpdateMouse(const TOptional<FIntPoint>& InMousePosition);
+	void UpdateMouse(const TOptional<FVector2D>& InMousePosition);
 
 	bool IsOverlayEnabled() const;
 	void ToggleOverlay();
diff --git a/Engine/Plugins/Experimental/MediaViewer/Source/MediaViewer/Private/Widgets/SMediaImageViewerStatusBar.cpp b/Engine/Plugins/Experimental/MediaViewer/Source/MediaViewer/Private/Widgets/SMediaImageViewerStatusBar.cpp
--- a/Engine/Plugins/Experimental/MediaViewer/Source/MediaViewer/Private/Widgets/SMediaImageViewerStatusBar.cpp
+++ b/Engine/Plugins/Experimental/MediaViewer/Source/MediaViewer/Private/Widgets/SMediaImageViewerStatusBar.cpp
@@ -7,7 +7,7 @@
 #include "ImageViewer/MediaImageViewer.h"
 #include "Internationalization/Text.h"
 #include "MediaViewerStyle.h"
-#include "SMediaImageViewerStatusBar.h"
+#include "Styling/StyleColors.h"
 #include "Widgets/MediaImageStatusBarExtender.h"
 #include "Widgets/SBoxPanel.h"
 #include "Widgets/SNullWidget.h"
@@ -68,13 +68,17 @@ TSharedRef<SWidget> SMediaImageViewerStatusBar::BuildStatusBar()
 	ApplyHook(StatusBarSections::StatusBarLeft, EExtensionHook::Before);
 	{
 		HorizontalBox->AddSlot()
-			.Padding(SlotPadding)
+			.Padding(SlotPadding + FMargin(0.f, 1.f, 0.f, 0.f))
 			.AutoWidth()
 			.HAlign(HAlign_Left)
+			.VAlign(VAlign_Center)
 			[
 				SNew(STextBlock)
 				.Text(this, &SMediaImageViewerStatusBar::GetResolutionLabel)
 				.Font(IDetailLayoutBuilder::GetDetailFont())
+				.ColorAndOpacity(FStyleColors::Foreground.GetSpecifiedColor())
+				.ShadowColorAndOpacity(FStyleColors::Panel.GetSpecifiedColor())
+				.ShadowOffset(FVector2D(1.f, 1.f))
 			];
 	}
 	ApplyHook(StatusBarSections::StatusBarLeft, EExtensionHook::After);
@@ -97,6 +101,7 @@ TSharedRef<SWidget> SMediaImageViewerStatusBar::BuildStatusBar()
 			.Padding(SlotPadding)
 			.AutoWidth()
 			.HAlign(HAlign_Right)
+			.VAlign(VAlign_Center)
 			[
 				SNew(SRichTextBlock)
 				.Text(this, &SMediaImageViewerStatusBar::GetColorPickerLabel)
@@ -150,7 +155,7 @@ FText SMediaImageViewerStatusBar::GetColorPickerLabel() const
 	if (!PixelColor.IsSet())
 	{
 		return FText::Format(
-			LOCTEXT("CoordinatesWithoutColor", "[{0}, {1}]"),
+			LOCTEXT("CoordinatesWithoutColor", "<RichTextBlock.Normal>[{0}, {1}]</>"),
 			FText::AsNumber(PixelCoordinates.X + 1),
 			FText::AsNumber(PixelCoordinates.Y + 1)
 		);
@@ -159,7 +164,7 @@ FText SMediaImageViewerStatusBar::GetColorPickerLabel() const
 	auto CoordsAndColorText = [&PixelCoordinates](const auto* Color, const FNumberFormattingOptions& Formatting)
 	{
 		return FText::Format(
-			LOCTEXT("CoordinatesWithColor", "[{0}, {1}] <RichTextBlock.Red>{2}</> <RichTextBlock.Green>{3}</> <RichTextBlock.Blue>{4}</> {5}"),
+			LOCTEXT("CoordinatesWithColor", "<RichTextBlock.Normal>[{0}, {1}]</> <RichTextBlock.Red>{2}</> <RichTextBlock.Green>{3}</> <RichTextBlock.Blue>{4}</> <RichTextBlock.Normal>{5}</>"),
 			FText::AsNumber(PixelCoordinates.X + 1),
 			FText::AsNumber(PixelCoordinates.Y + 1),
 			FText::AsNumber(Color->R, &Formatting),
diff --git a/Engine/Plugins/Experimental/MediaViewer/Source/MediaViewer/Private/Widgets/SMediaViewer.cpp b/Engine/Plugins/Experimental/MediaViewer/Source/MediaViewer/Private/Widgets/SMediaViewer.cpp
--- a/Engine/Plugins/Experimental/MediaViewer/Source/MediaViewer/Private/Widgets/SMediaViewer.cpp
+++ b/Engine/Plugins/Experimental/MediaViewer/Source/MediaViewer/Private/Widgets/SMediaViewer.cpp
@@ -267,15 +267,11 @@ void SMediaViewer::Tick(const FGeometry& AllottedGeometry, const double InCurren
 	const FVector2D ScreenSpacePosition = FSlateApplication::Get().GetCursorPos();
 	const FVector2D LocalPosition = AllottedGeometry.AbsoluteToLocal(ScreenSpacePosition);
 
-	FIntPoint CurrentLocation;
-	CurrentLocation.X = FMath::RoundToInt(LocalPosition.X);
-	CurrentLocation.Y = FMath::RoundToInt(LocalPosition.Y);
-
-	bool bNewLocation = CurrentLocation != CursorLocalPosition;
+	bool bNewLocation = (CursorLocalPosition - LocalPosition).IsNearlyZero() == false;
 
 	if (bNewLocation)
 	{
-		CursorLocalPosition = CurrentLocation;
+		CursorLocalPosition = LocalPosition;
 	}
 
 	if (bInvalidated)
@@ -284,9 +280,9 @@ void SMediaViewer::Tick(const FGeometry& AllottedGeometry, const double InCurren
 		return;
 	}
 
-	const TOptional<FIntPoint> UpdatedMousePosition = bNewLocation 
-		? CurrentLocation 
-		: TOptional<FIntPoint>();
+	const TOptional<FVector2D> UpdatedMousePosition = bNewLocation
+		? LocalPosition
+		: TOptional<FVector2D>();
 	
 	if (TSharedPtr<SMediaImageViewerOverlay> FirstOverlay = GetOverlay(EMediaImageViewerPosition::First))
 	{
@@ -673,7 +669,7 @@ EMediaImageViewerPosition SMediaViewer::GetHoveredImageViewer() const
 	}
 }
 
-FIntPoint SMediaViewer::GetLocalCursorPosition() const
+FVector2D SMediaViewer::GetLocalCursorPosition() const
 {
 	return CursorLocalPosition;
 }
diff --git a/Engine/Plugins/Experimental/MediaViewer/Source/MediaViewer/Private/Widgets/SMediaViewer.h b/Engine/Plugins/Experimental/MediaViewer/Source/MediaViewer/Private/Widgets/SMediaViewer.h
--- a/Engine/Plugins/Experimental/MediaViewer/Source/MediaViewer/Private/Widgets/SMediaViewer.h
+++ b/Engine/Plugins/Experimental/MediaViewer/Source/MediaViewer/Private/Widgets/SMediaViewer.h
@@ -82,7 +82,7 @@ public:
 	EMediaImageViewerPosition GetHoveredImageViewer() const;
 
 	/** Returns the local position of the cursor on the viewer. */
-	FIntPoint GetLocalCursorPosition() const;
+	FVector2D GetLocalCursorPosition() const;
 
 	/** Triggers a recreation of the view on tick. */
 	void InvalidateView();
@@ -130,7 +130,7 @@ protected:
 	/** The requested view configuration. */
 	EMediaImageViewerActivePosition RequestedView;
 
-	FIntPoint CursorLocalPosition;
+	FVector2D CursorLocalPosition;
 
 	mutable FVector2D ViewerSize;
 	mutable FVector2D ViewerPosition;
diff --git a/Engine/Plugins/Experimental/MediaViewer/Source/MediaViewer/Public/ImageViewer/MediaImageViewer.h b/Engine/Plugins/Experimental/MediaViewer/Source/MediaViewer/Public/ImageViewer/MediaImageViewer.h
--- a/Engine/Plugins/Experimental/MediaViewer/Source/MediaViewer/Public/ImageViewer/MediaImageViewer.h
+++ b/Engine/Plugins/Experimental/MediaViewer/Source/MediaViewer/Public/ImageViewer/MediaImageViewer.h
@@ -243,7 +243,7 @@ public:
 	 * @param   MousePosition - The position of the mouse when tracking stoppped.
 	 * @returns True if custom tracking was stopped.
 	 */
-	virtual void OnMouseMove(const FIntPoint& InMousePosition) {}
+	virtual void OnMouseMove(const FVector2D& InMousePosition) {}
 
 	/**
 	 * Responds to key input in this image's viewport.
