Media Viewer: Small style fix

diff --git a/Engine/Plugins/Experimental/MediaViewer/Source/MediaViewer/Private/MediaViewerStyle.cpp b/Engine/Plugins/Experimental/MediaViewer/Source/MediaViewer/Private/MediaViewerStyle.cpp
--- a/Engine/Plugins/Experimental/MediaViewer/Source/MediaViewer/Private/MediaViewerStyle.cpp
+++ b/Engine/Plugins/Experimental/MediaViewer/Source/MediaViewer/Private/MediaViewerStyle.cpp
@@ -23,7 +23,7 @@ FMediaViewerStyle::FMediaViewerStyle()
 
 	FTextBlockStyle NormalText = FAppStyle::Get().GetWidgetStyle<FTextBlockStyle>("NormalText");
 	NormalText.SetFont(IDetailLayoutBuilder::GetDetailFont());
-	NormalText.SetShadowColorAndOpacity(FStyleColors::Foreground.GetSpecifiedColor());
+	NormalText.SetShadowColorAndOpacity(FStyleColors::Panel.GetSpecifiedColor());
 	NormalText.SetShadowOffset(FVector2D(1.0, 1.0));
 	Set("RichTextBlock.Normal", NormalText);
 	Set("RichTextBlock.Red", FTextBlockStyle(NormalText).SetColorAndOpacity(FLinearColor(1.0f, 0.1f, 0.1f)));
