Media Viewer: Slate alignment

diff --git a/Engine/Plugins/Experimental/MediaViewer/Source/MediaViewer/Private/Widgets/SMediaImageViewerStatusBar.cpp b/Engine/Plugins/Experimental/MediaViewer/Source/MediaViewer/Private/Widgets/SMediaImageViewerStatusBar.cpp
--- a/Engine/Plugins/Experimental/MediaViewer/Source/MediaViewer/Private/Widgets/SMediaImageViewerStatusBar.cpp
+++ b/Engine/Plugins/Experimental/MediaViewer/Source/MediaViewer/Private/Widgets/SMediaImageViewerStatusBar.cpp
@@ -68,7 +68,7 @@ TSharedRef<SWidget> SMediaImageViewerStatusBar::BuildStatusBar()
 	ApplyHook(StatusBarSections::StatusBarLeft, EExtensionHook::Before);
 	{
 		HorizontalBox->AddSlot()
-			.Padding(SlotPadding + FMargin(0.f, 1.f, 0.f, 0.f))
+			.Padding(SlotPadding)
 			.AutoWidth()
 			.HAlign(HAlign_Left)
 			.VAlign(VAlign_Center)
