Media Viewer: Media Source integration

diff --git a/Engine/Plugins/Experimental/MediaViewer/Source/MediaViewer/MediaViewer.build.cs b/Engine/Plugins/Experimental/MediaViewer/Source/MediaViewer/MediaViewer.build.cs
--- a/Engine/Plugins/Experimental/MediaViewer/Source/MediaViewer/MediaViewer.build.cs
+++ b/Engine/Plugins/Experimental/MediaViewer/Source/MediaViewer/MediaViewer.build.cs
@@ -28,6 +28,8 @@ public class MediaViewer : ModuleRules
 				"InputCore",
 				"LevelEditor",
 				"MediaAssets",
+				"MediaPlayerEditor",
+				"MediaStream",
 				"Projects",
 				"PropertyEditor",
 				"RenderCore",
diff --git a/Engine/Plugins/Experimental/MediaViewer/Source/MediaViewer/Private/ImageViewers/MediaSourceImageViewer.cpp b/Engine/Plugins/Experimental/MediaViewer/Source/MediaViewer/Private/ImageViewers/MediaSourceImageViewer.cpp
new file mode 100644
--- /dev/null
+++ b/Engine/Plugins/Experimental/MediaViewer/Source/MediaViewer/Private/ImageViewers/MediaSourceImageViewer.cpp
@@ -0,0 +1,351 @@
+// Copyright Epic Games, Inc. All Rights Reserved.
+
+#include "ImageViewers/MediaSourceImageViewer.h"
+
+#include "AssetRegistry/AssetData.h"
+#include "Brushes/SlateImageBrush.h"
+#include "Engine/Texture2D.h"
+#include "IMediaStreamPlayer.h"
+#include "Materials/MaterialInstanceDynamic.h"
+#include "Materials/MaterialInterface.h"
+#include "Math/Vector2D.h"
+#include "MediaPlayer.h"
+#include "MediaSource.h"
+#include "MediaStream.h"
+#include "MediaStreamSourceBlueprintLibrary.h"
+#include "MediaTexture.h"
+#include "MediaViewerUtils.h"
+#include "Misc/Optional.h"
+#include "Misc/TVariant.h"
+#include "Widgets/Layout/SBox.h"
+#include "Widgets/SMediaSourceOverlay.h"
+
+#define LOCTEXT_NAMESPACE "MediaSourceImageViewer"
+
+namespace UE::MediaViewer::Private
+{
+
+constexpr const TCHAR* MediaSourceThumbnailTexturePath = TEXT("/Script/Engine.Texture2D'/Engine/EditorResources/SceneManager.SceneManager'");
+
+UTexture2D* GetMediaSourceThumbnailTexture()
+{
+	return LoadObject<UTexture2D>(GetTransientPackage(), MediaSourceThumbnailTexturePath);
+}
+
+const FLazyName FMediaSourceImageViewer::ItemTypeName = TEXT("MediaSource");
+
+bool FMediaSourceImageViewer::FFactory::SupportsAsset(const FAssetData& InAssetData) const
+{
+	if (UClass* Class = InAssetData.GetClass(EResolveClass::Yes))
+	{
+		return Class->IsChildOf<UMediaSource>();
+	}
+
+	return false;
+}
+
+TSharedPtr<FMediaImageViewer> FMediaSourceImageViewer::FFactory::CreateImageViewer(const FAssetData& InAssetData) const
+{
+	if (UObject* Object = InAssetData.GetAsset())
+	{
+		return CreateImageViewer(TNotNull<UObject*>(Object));
+	}
+
+	return nullptr;
+}
+
+TSharedPtr<FMediaViewerLibraryItem> FMediaSourceImageViewer::FFactory::CreateLibraryItem(const FAssetData& InAssetData) const
+{
+	if (UObject* Object = InAssetData.GetAsset())
+	{
+		return CreateLibraryItem(TNotNull<UObject*>(Object));
+	}
+
+	return nullptr;
+}
+
+bool FMediaSourceImageViewer::FFactory::SupportsObject(TNotNull<UObject*> InObject) const
+{
+	return InObject->IsA<UMediaSource>();
+}
+
+TSharedPtr<FMediaImageViewer> FMediaSourceImageViewer::FFactory::CreateImageViewer(TNotNull<UObject*> InObject) const
+{
+	if (UMediaSource* MediaSource = Cast<UMediaSource>(InObject))
+	{
+		return MakeShared<FMediaSourceImageViewer>(MediaSource);
+	}
+
+	return nullptr;
+}
+
+TSharedPtr<FMediaViewerLibraryItem> FMediaSourceImageViewer::FFactory::CreateLibraryItem(TNotNull<UObject*> InObject) const
+{
+	if (UMediaSource* MediaSource = Cast<UMediaSource>(InObject))
+	{
+		return MakeShared<FMediaSourceImageViewer::FItem>(
+			FText::FromName(MediaSource->GetFName()),
+			FText::Format(LOCTEXT("FactoryToolTipFormat", "{0} [Media Source]"), FText::FromString(MediaSource->GetPathName())),
+			MediaSource->HasAnyFlags(EObjectFlags::RF_Transient) || MediaSource->IsIn(GetTransientPackage()),
+			MediaSource
+		);
+	}
+
+	return nullptr;	
+}
+
+bool FMediaSourceImageViewer::FFactory::SupportsItemType(FName InItemType) const
+{
+	return InItemType == FMediaSourceImageViewer::ItemTypeName;
+}
+
+TSharedPtr<FMediaViewerLibraryItem> FMediaSourceImageViewer::FFactory::CreateLibraryItem(const FMediaViewerLibraryItem& InSavedItem) const
+{
+	return MakeShared<FMediaSourceImageViewer::FItem>(FPrivateToken(), InSavedItem);
+}
+
+FMediaSourceImageViewer::FItem::FItem(const FText& InName, const FText& InToolTip, bool bInTransient, TNotNull<UMediaSource*> InMediaSource)
+	: FItem(FGuid::NewGuid(), InName, InToolTip, bInTransient, InMediaSource)
+{
+}
+
+FMediaSourceImageViewer::FItem::FItem(const FGuid& InId, const FText& InName, const FText& InToolTip, bool bInTransient,
+	TNotNull<UMediaSource*> InMediaSource)
+	: FMediaViewerLibraryItem(InId, InName, InToolTip, bInTransient, InMediaSource->GetPathName())
+{
+}
+
+FMediaSourceImageViewer::FItem::FItem(FPrivateToken&& InPrivateToken, const FMediaViewerLibraryItem& InItem)
+	: FMediaViewerLibraryItem(InItem.GetId(), InItem.Name, InItem.ToolTip, InItem.IsTransient(), InItem.GetStringValue())
+{
+}
+
+FName FMediaSourceImageViewer::FItem::GetItemType() const
+{
+	return FMediaSourceImageViewer::ItemTypeName;
+}
+
+FText FMediaSourceImageViewer::FItem::GetItemTypeDisplayName() const
+{
+	return LOCTEXT("MediaSource", "Media");
+}
+
+TSharedPtr<FSlateBrush> FMediaSourceImageViewer::FItem::CreateThumbnail()
+{
+	// Loading the media to create a thumbnail is not a good idea.
+	TSharedRef<FSlateImageBrush> ThumbnailBrush = MakeShared<FSlateImageBrush>(
+		GetMediaSourceThumbnailTexture(), 
+		FVector2D(64.0, 64.0)
+	);
+
+	return ThumbnailBrush;
+}
+
+TSharedPtr<FMediaImageViewer> FMediaSourceImageViewer::FItem::CreateImageViewer() const
+{
+	UMediaSource* MediaSource = LoadAssetFromString<UMediaSource>(StringValue);
+
+	if (!MediaSource)
+	{
+		return nullptr;
+	}
+
+	if (Id.IsValid())
+	{
+		return MakeShared<FMediaSourceImageViewer>(Id, MediaSource);
+	}
+
+	return MakeShared<FMediaSourceImageViewer>(MediaSource);
+}
+
+TSharedPtr<FMediaViewerLibraryItem> FMediaSourceImageViewer::FItem::Clone() const
+{
+	if (StringValue.IsEmpty())
+	{
+		return nullptr;
+	}
+
+	return MakeShared<FMediaSourceImageViewer::FItem>(FPrivateToken(), *this);
+}
+
+FString FMediaSourceImageViewer::FItem::GetReferencerName() const
+{
+	static const FString ReferencerName = TEXT("FMediaSourceImageViewer::FItem");
+	return ReferencerName;
+}
+
+void FMediaSourceImageViewer::FItem::AddReferencedObjects(FReferenceCollector& Collector)
+{
+	if (Texture)
+	{
+		Collector.AddReferencedObject(Texture);
+	}
+}
+
+FMediaSourceImageViewer::FMediaSourceImageViewer(TNotNull<UMediaSource*> InMediaSource)
+	: FMediaSourceImageViewer(FGuid::NewGuid(), InMediaSource)
+{
+}
+
+FMediaSourceImageViewer::FMediaSourceImageViewer(const FGuid& InId, TNotNull<UMediaSource*> InMediaSource)
+	: FMediaImageViewer({
+		InId,
+		FIntPoint(1, 1),
+		1,
+		FText::FromName(InMediaSource->GetFName())
+	})
+{
+	MediaSourceSettings.MediaSource = InMediaSource;
+
+	MediaStream = NewObject<UMediaStream>(GetTransientPackage());
+
+	TSoftObjectPtr<UObject> SoftPtr(InMediaSource);
+	FMediaStreamSource Source = UMediaStreamSourceBlueprintLibrary::MakeMediaSourceFromAsset(MediaStream, SoftPtr);
+
+	if (UMediaStreamSourceBlueprintLibrary::IsValidMediaSource(Source))
+	{
+		MediaStream->SetSource(Source);
+
+		if (TScriptInterface<IMediaStreamPlayer> Player = MediaStream->GetPlayer())
+		{
+			if (UMediaTexture* MediaTexture = Player->GetMediaTexture())
+			{
+				ImageInfo.Size.X = MediaTexture->GetSurfaceWidth();
+				ImageInfo.Size.Y = MediaTexture->GetSurfaceHeight();
+
+				FMediaStreamPlayerConfig Config = Player->GetPlayerConfig();
+				Config.bPlayOnOpen = true;
+				Config.bLooping = true;
+				Player->SetPlayerConfig(Config);
+
+				Player->Play();
+				Brush = MakeShared<FSlateImageBrush>(MediaTexture, FVector2D(ImageInfo.Size.X, ImageInfo.Size.Y));
+
+				SampleCache = FTextureSampleCache(MediaTexture);
+			}
+		}
+	}
+}
+
+UMediaStream* FMediaSourceImageViewer::GetMediaStream() const
+{
+	return MediaStream;
+}
+
+TSharedPtr<FMediaViewerLibraryItem> FMediaSourceImageViewer::CreateLibraryItem() const
+{
+	UMediaSource* MediaSource = MediaSourceSettings.MediaSource;
+
+	if (!MediaSource)
+	{
+		return nullptr;
+	}
+
+	return MakeShared<FItem>(
+		ImageInfo.Id,
+		FText::FromName(MediaSource->GetFName()),
+		FText::Format(LOCTEXT("LibraryToolTipFormat", "{0} [Media Texture - {1}x{2}]"), FText::FromString(MediaSource->GetPathName()), FText::AsNumber(ImageInfo.Size.X), FText::AsNumber(ImageInfo.Size.Y)),
+		MediaSource->HasAnyFlags(EObjectFlags::RF_Transient) || MediaSource->IsIn(GetTransientPackage()),
+		MediaSource
+	);
+}
+
+TOptional<TVariant<FColor, FLinearColor>> FMediaSourceImageViewer::GetPixelColor(const FIntPoint& InPixelCoords, int32 InMipLevel) const
+{
+	if (InPixelCoords.X < 0 || InPixelCoords.Y < 0)
+	{
+		SampleCache.Invalidate();
+		return {};
+	}
+
+	if (InPixelCoords.X >= ImageInfo.Size.X || InPixelCoords.Y >= ImageInfo.Size.Y)
+	{
+		SampleCache.Invalidate();
+		return {};
+	}
+
+	if (!MediaStream)
+	{
+		SampleCache.Invalidate();
+		return {};
+	}
+
+	IMediaStreamPlayer* MediaStreamPlayer = MediaStream->GetPlayer().GetInterface();
+
+	if (!MediaStreamPlayer)
+	{
+		SampleCache.Invalidate();
+		return {};
+	}
+
+	if (UMediaPlayer* Player = MediaStreamPlayer->GetPlayer())
+	{
+		const FTimespan PlayerTime = Player->GetTime();
+
+		if (const FLinearColor* PixelColor = SampleCache.GetPixelColor(InPixelCoords, PlayerTime))
+		{
+			TVariant<FColor, FLinearColor> PixelColorVariant;
+			PixelColorVariant.Set<FLinearColor>(*PixelColor);
+
+			return PixelColorVariant;
+		}
+	}
+
+	return {};
+}
+
+TSharedPtr<FStructOnScope> FMediaSourceImageViewer::GetCustomSettingsOnScope() const
+{
+	return MakeShared<FStructOnScope>(
+		FMediaSourceImageViewerSettings::StaticStruct(),
+		reinterpret_cast<uint8*>(&const_cast<FMediaSourceImageViewerSettings&>(MediaSourceSettings))
+	);
+}
+
+void FMediaSourceImageViewer::AddReferencedObjects(FReferenceCollector& Collector)
+{
+	Collector.AddPropertyReferencesWithStructARO(
+		FMediaSourceImageViewerSettings::StaticStruct(),
+		&MediaSourceSettings
+	);
+
+	if (MediaStream)
+	{
+		Collector.AddReferencedObject(MediaStream);
+	}
+}
+
+TSharedPtr<SWidget> FMediaSourceImageViewer::GetOverlayWidget(EMediaImageViewerPosition InPosition)
+{
+	return SNew(SBox)
+		.HAlign(HAlign_Fill)
+		.VAlign(VAlign_Bottom)
+		.Padding(5.f)
+		[
+			SNew(SMediaSourceOverlay, SharedThis(this))
+		];
+}
+
+FString FMediaSourceImageViewer::GetReferencerName() const
+{
+	static const FString ReferencerName = TEXT("FMediaSourceImageViewer");
+	return ReferencerName;
+}
+
+void FMediaSourceImageViewer::PaintImage(FMediaImagePaintParams& InPaintParams, const FMediaImagePaintGeometry& InPaintGeometry)
+{
+	if (TScriptInterface<IMediaStreamPlayer> Player = MediaStream->GetPlayer())
+	{
+		if (UMediaTexture* MediaTexture = Player->GetMediaTexture())
+		{
+			ImageInfo.Size.X = MediaTexture->GetSurfaceWidth();
+			ImageInfo.Size.Y = MediaTexture->GetSurfaceHeight();
+		}
+	}
+
+	FMediaImageViewer::PaintImage(InPaintParams, InPaintGeometry);
+}
+
+} // UE::MediaViewer::Private
+
+#undef LOCTEXT_NAMESPACE
diff --git a/Engine/Plugins/Experimental/MediaViewer/Source/MediaViewer/Private/ImageViewers/MediaSourceImageViewer.h b/Engine/Plugins/Experimental/MediaViewer/Source/MediaViewer/Private/ImageViewers/MediaSourceImageViewer.h
new file mode 100644
--- /dev/null
+++ b/Engine/Plugins/Experimental/MediaViewer/Source/MediaViewer/Private/ImageViewers/MediaSourceImageViewer.h
@@ -0,0 +1,112 @@
+// Copyright Epic Games, Inc. All Rights Reserved.
+
+#pragma once
+
+#include "ImageViewer/IMediaImageViewerFactory.h"
+#include "ImageViewer/MediaImageViewer.h"
+#include "Library/MediaViewerLibraryItem.h"
+#include "UObject/GCObject.h"
+
+#include "ImageViewers/TextureSampleCache.h"
+#include "Math/Color.h"
+#include "Misc/Optional.h"
+#include "Misc/Timespan.h"
+#include "Templates/SharedPointer.h"
+
+#include "MediaSourceImageViewer.generated.h"
+
+class UMediaSource;
+class UMediaStream;
+
+USTRUCT()
+struct FMediaSourceImageViewerSettings
+{
+	GENERATED_BODY()
+
+	UPROPERTY(VisibleAnywhere, Category = "Texture")
+	TObjectPtr<UMediaSource> MediaSource;
+};
+
+namespace UE::MediaViewer::Private
+{
+
+class FMediaSourceImageViewer : public FMediaImageViewer, public FGCObject
+{
+public:
+	struct FFactory : public IMediaImageViewerFactory
+	{
+		FFactory()
+		{
+			Priority = 5000;
+		}
+
+		virtual bool SupportsAsset(const FAssetData& InAssetData) const override;
+		virtual TSharedPtr<FMediaImageViewer> CreateImageViewer(const FAssetData& InAssetData) const override;
+		virtual TSharedPtr<FMediaViewerLibraryItem> CreateLibraryItem(const FAssetData& InAssetData) const override;
+
+		virtual bool SupportsObject(TNotNull<UObject*> InObject) const override;
+		virtual TSharedPtr<FMediaImageViewer> CreateImageViewer(TNotNull<UObject*> InObject) const override;
+		virtual TSharedPtr<FMediaViewerLibraryItem> CreateLibraryItem(TNotNull<UObject*> InObject) const override;
+
+		virtual bool SupportsItemType(FName InItemType) const override;
+		virtual TSharedPtr<FMediaViewerLibraryItem> CreateLibraryItem(const FMediaViewerLibraryItem& InSavedItem) const override;
+	};
+
+	struct FItem : public FMediaViewerLibraryItem, public FGCObject
+	{
+		TObjectPtr<UTexture> Texture;
+
+		FItem(const FText& InName, const FText& InToolTip, bool bInTransient, TNotNull<UMediaSource*> InMediaSource);
+
+		FItem(const FGuid& InId, const FText& InName, const FText& InToolTip, bool bInTransient, TNotNull<UMediaSource*> InMediaSource);
+
+		FItem(FPrivateToken&& InPrivateToken, const FMediaViewerLibraryItem& InItem);
+
+		virtual ~FItem() override = default;
+
+		//~ Begin FMediaViewerLibraryItem
+		virtual FName GetItemType() const override;
+		virtual FText GetItemTypeDisplayName() const override;
+		virtual TSharedPtr<FSlateBrush> CreateThumbnail() override;
+		virtual TSharedPtr<FMediaImageViewer> CreateImageViewer() const override;
+		virtual TSharedPtr<FMediaViewerLibraryItem> Clone() const override;
+		//~ End FMediaViewerLibraryItem
+
+		//~ Begin FGCObject
+		virtual FString GetReferencerName() const override;
+		virtual void AddReferencedObjects(FReferenceCollector& Collector) override;
+		//~ End FGCObject
+	};
+
+	static const FLazyName ItemTypeName;
+
+	FMediaSourceImageViewer(TNotNull<UMediaSource*> InMediaSource);
+	FMediaSourceImageViewer(const FGuid& InId, TNotNull<UMediaSource*> InMediaSource);
+
+	virtual ~FMediaSourceImageViewer() override = default;
+
+	UMediaStream* GetMediaStream() const;
+
+	//~ Begin FMediaImageViewer
+	virtual TSharedPtr<FMediaViewerLibraryItem> CreateLibraryItem() const override;
+	virtual TOptional<TVariant<FColor, FLinearColor>> GetPixelColor(const FIntPoint& InPixelCoords, int32 InMipLevel) const override;
+	virtual TSharedPtr<FStructOnScope> GetCustomSettingsOnScope() const override;
+	virtual TSharedPtr<SWidget> GetOverlayWidget(EMediaImageViewerPosition InPosition) override;
+	//~ End FMediaImageViewer
+
+	//~ Begin FGCObject
+	virtual FString GetReferencerName() const override;
+	virtual void AddReferencedObjects(FReferenceCollector& Collector) override;
+	//~ End FGCObject
+
+protected:
+	FMediaSourceImageViewerSettings MediaSourceSettings;
+	TObjectPtr<UMediaStream> MediaStream;
+	mutable FTextureSampleCache SampleCache;
+
+	//~ Begin FMediaImageViewer
+	virtual void PaintImage(FMediaImagePaintParams& InPaintParams, const FMediaImagePaintGeometry& InPaintGeometry) override;
+	//~ End FMediaImageViewer
+};
+
+} // UE::MediaViewer::Private
diff --git a/Engine/Plugins/Experimental/MediaViewer/Source/MediaViewer/Private/ImageViewers/TextureMipCache.cpp b/Engine/Plugins/Experimental/MediaViewer/Source/MediaViewer/Private/ImageViewers/TextureMipCache.cpp
--- a/Engine/Plugins/Experimental/MediaViewer/Source/MediaViewer/Private/ImageViewers/TextureMipCache.cpp
+++ b/Engine/Plugins/Experimental/MediaViewer/Source/MediaViewer/Private/ImageViewers/TextureMipCache.cpp
@@ -72,4 +72,4 @@ void FTextureMipCache::CreateMipImage(int32 InMipLevel)
 		: EMipState::Failed;	
 }
 
-}; // UE::MediaViewer::Private
+} // UE::MediaViewer::Private
diff --git a/Engine/Plugins/Experimental/MediaViewer/Source/MediaViewer/Private/ImageViewers/TextureSampleCache.cpp b/Engine/Plugins/Experimental/MediaViewer/Source/MediaViewer/Private/ImageViewers/TextureSampleCache.cpp
new file mode 100644
--- /dev/null
+++ b/Engine/Plugins/Experimental/MediaViewer/Source/MediaViewer/Private/ImageViewers/TextureSampleCache.cpp
@@ -0,0 +1,73 @@
+// Copyright Epic Games, Inc. All Rights Reserved.
+
+#include "ImageViewers/TextureSampleCache.h"
+
+#include "Engine/Texture.h"
+#include "RenderingThread.h"
+#include "TextureResource.h"
+
+namespace UE::MediaViewer::Private
+{
+
+FTextureSampleCache::FTextureSampleCache()
+	: Texture(nullptr)
+{
+}
+
+FTextureSampleCache::FTextureSampleCache(TNotNull<UTexture*> InTexture)
+	: Texture(InTexture)
+{
+}
+
+const FLinearColor* FTextureSampleCache::GetPixelColor(const FIntPoint& InPixelCoordinates, const TOptional<FTimespan>& InTime)
+{
+	if (!Texture)
+	{
+		return nullptr;
+	}
+
+	const bool bNeedsUpdate = !PixelColorSample.IsSet() || InPixelCoordinates != PixelColorSample->Coordinates
+		|| (InTime.IsSet() && (PixelColorSample->Time != InTime.GetValue()));
+
+	if (bNeedsUpdate)
+	{
+		ENQUEUE_RENDER_COMMAND(GetPixelColors)
+		(
+			[this, InPixelCoordinates, InTime](FRHICommandListImmediate& RHICmdList)
+			{
+				PixelColorSample.Reset();
+
+				if (FTextureResource* TextureResource = Texture->GetResource())
+				{
+					const FIntVector Size = TextureResource->TextureRHI->GetSizeXYZ();
+					TArray<FLinearColor> Data;
+					FIntRect Rect(InPixelCoordinates.X, InPixelCoordinates.Y, InPixelCoordinates.X + 1, InPixelCoordinates.Y + 1);
+					RHICmdList.ReadSurfaceData(TextureResource->TextureRHI, Rect, Data, FReadSurfaceDataFlags());
+
+					if (!Data.IsEmpty())
+					{
+						PixelColorSample = {
+							InPixelCoordinates,
+							InTime,
+							Data[0]
+						};
+					}
+				}
+			}
+		);
+	}
+
+	if (PixelColorSample.IsSet())
+	{
+		return &PixelColorSample->Color;
+	}
+
+	return nullptr;
+}
+
+void FTextureSampleCache::Invalidate()
+{
+	PixelColorSample.Reset();
+}
+
+} // UE::MediaViewer::Private
diff --git a/Engine/Plugins/Experimental/MediaViewer/Source/MediaViewer/Private/ImageViewers/TextureSampleCache.h b/Engine/Plugins/Experimental/MediaViewer/Source/MediaViewer/Private/ImageViewers/TextureSampleCache.h
new file mode 100644
--- /dev/null
+++ b/Engine/Plugins/Experimental/MediaViewer/Source/MediaViewer/Private/ImageViewers/TextureSampleCache.h
@@ -0,0 +1,38 @@
+// Copyright Epic Games, Inc. All Rights Reserved.
+
+#pragma once
+
+#include "Math/Color.h"
+#include "Math/IntPoint.h"
+#include "Misc/NotNull.h"
+#include "Misc/Optional.h"
+#include "Misc/Timespan.h"
+
+class UTexture;
+
+namespace UE::MediaViewer::Private
+{
+
+class FTextureSampleCache
+{
+public:
+	FTextureSampleCache();
+	FTextureSampleCache(TNotNull<UTexture*> InTexture);
+
+	const FLinearColor* GetPixelColor(const FIntPoint& InPixelCoordinates, const TOptional<FTimespan>& InTime = {});
+
+	void Invalidate();
+
+protected:
+	struct FPixelColorSample
+	{
+		FIntPoint Coordinates;
+		TOptional<FTimespan> Time;
+		FLinearColor Color;
+	};
+
+	UTexture* Texture;
+	TOptional<FPixelColorSample> PixelColorSample;
+};
+
+} // UE::MediaViewer::Private
diff --git a/Engine/Plugins/Experimental/MediaViewer/Source/MediaViewer/Private/MediaViewerModule.cpp b/Engine/Plugins/Experimental/MediaViewer/Source/MediaViewer/Private/MediaViewerModule.cpp
--- a/Engine/Plugins/Experimental/MediaViewer/Source/MediaViewer/Private/MediaViewerModule.cpp
+++ b/Engine/Plugins/Experimental/MediaViewer/Source/MediaViewer/Private/MediaViewerModule.cpp
@@ -6,6 +6,7 @@
 #include "Framework/Docking/TabManager.h"
 #include "ImageViewers/ColorImageViewer.h"
 #include "ImageViewers/MaterialInterfaceImageViewer.h"
+#include "ImageViewers/MediaSourceImageViewer.h"
 #include "ImageViewers/MediaTextureImageViewer.h"
 #include "ImageViewers/NullImageViewer.h"
 #include "ImageViewers/Texture2DImageViewer.h"
@@ -474,6 +475,7 @@ void FMediaViewerModule::RegisterDefaultImageViewers()
 {
 	RegisterFactory(FColorImageViewer::ItemTypeName,                 MakeShared<FColorImageViewer::FFactory>());
 	RegisterFactory(FMaterialInterfaceImageViewer::ItemTypeName,     MakeShared<FMaterialInterfaceImageViewer::FFactory>());
+	RegisterFactory(FMediaSourceImageViewer::ItemTypeName,           MakeShared<FMediaSourceImageViewer::FFactory>());
 	RegisterFactory(FMediaTextureImageViewer::ItemTypeName,          MakeShared<FMediaTextureImageViewer::FFactory>());
 	RegisterFactory(FTexture2DImageViewer::ItemTypeName,             MakeShared<FTexture2DImageViewer::FFactory>());
 	RegisterFactory(FTextureRenderTarget2DImageViewer::ItemTypeName, MakeShared<FTextureRenderTarget2DImageViewer::FFactory>());
diff --git a/Engine/Plugins/Experimental/MediaViewer/Source/MediaViewer/Private/Widgets/SMediaSourceOverlay.cpp b/Engine/Plugins/Experimental/MediaViewer/Source/MediaViewer/Private/Widgets/SMediaSourceOverlay.cpp
new file mode 100644
--- /dev/null
+++ b/Engine/Plugins/Experimental/MediaViewer/Source/MediaViewer/Private/Widgets/SMediaSourceOverlay.cpp
@@ -0,0 +1,545 @@
+// Copyright Epic Games, Inc. All Rights Reserved.
+
+#include "Widgets/SMediaSourceOverlay.h"
+
+#include "DetailLayoutBuilder.h"
+#include "ImageViewers/MediaSourceImageViewer.h"
+#include "IMediaStreamPlayer.h"
+#include "MediaPlayer.h"
+#include "MediaPlayerEditorModule.h"
+#include "MediaStream.h"
+#include "Misc/FrameRate.h"
+#include "Misc/Timecode.h"
+#include "Misc/Timespan.h"
+#include "Modules/ModuleManager.h"
+#include "Styling/SlateIconFinder.h"
+#include "Styling/StyleColors.h"
+#include "Widgets/Images/SImage.h"
+#include "Widgets/Input/SButton.h"
+#include "Widgets/SBoxPanel.h"
+#include "Widgets/SNullWidget.h"
+#include "Widgets/Text/STextBlock.h"
+
+#define LOCTEXT_NAMESPACE "SMediaSourceOverlay"
+
+namespace UE::MediaViewer::Private
+{
+
+void SMediaSourceOverlay::PrivateRegisterAttributes(FSlateAttributeDescriptor::FInitializer&)
+{
+}
+
+void SMediaSourceOverlay::Construct(const FArguments& InArgs, const TSharedRef<FMediaSourceImageViewer>& InImageViewer)
+{
+	ImageViewerWeak = InImageViewer;
+
+	TrySetFrameRate();
+
+	ChildSlot
+	[
+		SNew(SVerticalBox)
+		
+		+ SVerticalBox::Slot()
+		.AutoHeight()
+		.HAlign(HAlign_Fill)
+		.MaxSize(400.f)
+		[
+			CreateSlider()
+		]
+
+		+ SVerticalBox::Slot()
+		.AutoHeight()
+		.HAlign(HAlign_Fill)
+		.Padding(0.f, 5.f, 0.f, 0.f)
+		[
+			CreateControls()
+		]
+	];
+}
+
+UMediaStream* SMediaSourceOverlay::GetMediaStream() const
+{
+	if (TSharedPtr<FMediaSourceImageViewer> ImageViewer = ImageViewerWeak.Pin())
+	{
+		return ImageViewer->GetMediaStream();
+	}
+
+	return nullptr;
+}
+
+IMediaStreamPlayer* SMediaSourceOverlay::GetMediaStreamPlayer() const
+{
+	if (UMediaStream* MediaStream = GetMediaStream())
+	{
+		return MediaStream->GetPlayer().GetInterface();
+	}
+
+	return nullptr;
+}
+
+UMediaPlayer* SMediaSourceOverlay::GetMediaPlayer() const
+{
+	if (IMediaStreamPlayer* MediaStreamPlayer = GetMediaStreamPlayer())
+	{
+		return MediaStreamPlayer->GetPlayer();
+	}
+
+	return nullptr;
+}
+
+void SMediaSourceOverlay::TrySetFrameRate()
+{
+	if (UMediaPlayer* MediaPlayer = GetMediaPlayer())
+	{
+		// IDEX_NONE here access the currently playing TrackIndex and FormatIndex.
+		const float FrameRate = MediaPlayer->GetVideoTrackFrameRate(INDEX_NONE, INDEX_NONE);
+
+		if (!FMath::IsNearlyZero(FrameRate))
+		{
+			FrameRateFloat = FrameRate;
+		}
+	}
+}
+
+TSharedRef<SWidget> SMediaSourceOverlay::CreateSlider()
+{
+	UMediaPlayer* MediaPlayer = GetMediaPlayer();
+
+	if (!MediaPlayer)
+	{
+		return SNullWidget::NullWidget;
+	}
+
+	IMediaPlayerEditorModule* MediaPlayerEditorModule = FModuleManager::LoadModulePtr<IMediaPlayerEditorModule>("MediaPlayerEditor");
+
+	if (!MediaPlayerEditorModule)
+	{
+		return SNullWidget::NullWidget;
+	}
+
+	TArray<TWeakObjectPtr<UMediaPlayer>> MediaPlayerList;
+	MediaPlayerList.Add(MediaPlayer);
+
+	const TSharedRef<IMediaPlayerSlider> MediaPlayerSlider = MediaPlayerEditorModule->CreateMediaPlayerSliderWidget(MediaPlayerList);
+
+	MediaPlayerSlider->SetSliderHandleColor(FSlateColor(EStyleColor::AccentBlue));
+	MediaPlayerSlider->SetVisibleWhenInactive(EVisibility::Visible);
+
+	return MediaPlayerSlider;
+}
+
+TSharedRef<SWidget> SMediaSourceOverlay::CreateControls()
+{
+	return SNew(SHorizontalBox)
+
+		// Current frame
+		+SHorizontalBox::Slot()
+		.HAlign(HAlign_Right)
+		.VAlign(VAlign_Center)
+		.Padding(0.f, 0.f, 5.f, 0.f)
+		.FillWidth(1.f)
+		[
+			SNew(STextBlock)
+			.ColorAndOpacity(FStyleColors::Foreground.GetSpecifiedColor())
+			.ShadowColorAndOpacity(FStyleColors::Panel.GetSpecifiedColor())
+			.ShadowOffset(FVector2D(1.f, 1.f))
+			.Font(IDetailLayoutBuilder::GetDetailFont())
+			.Text(this, &SMediaSourceOverlay::GetCurrentFrame)
+		]
+
+		+ SHorizontalBox::Slot()
+		.VAlign(VAlign_Center)
+		.AutoWidth()
+		.Padding(0.f, 0.f, 5.f, 0.f)
+		[
+			SNew(STextBlock)
+			.ColorAndOpacity(FStyleColors::Foreground.GetSpecifiedColor())
+			.ShadowColorAndOpacity(FStyleColors::Panel.GetSpecifiedColor())
+			.ShadowOffset(FVector2D(1.f, 1.f))
+			.Font(IDetailLayoutBuilder::GetDetailFont())
+			.Text(INVTEXT("/"))
+		]
+
+		// Total fames
+		+ SHorizontalBox::Slot()
+		.VAlign(VAlign_Center)
+		.AutoWidth()
+		.Padding(0.f, 0.f, 10.f, 0.f)
+		[
+			SNew(STextBlock)
+			.ColorAndOpacity(FStyleColors::Foreground.GetSpecifiedColor())
+			.ShadowColorAndOpacity(FStyleColors::Panel.GetSpecifiedColor())
+			.ShadowOffset(FVector2D(1.f, 1.f))
+			.Font(IDetailLayoutBuilder::GetDetailFont())
+			.Text(this, &SMediaSourceOverlay::GetTotalFrames)
+		]
+
+		// Rewind button.
+		+ SHorizontalBox::Slot()
+		.VAlign(VAlign_Center)
+		.AutoWidth()
+		[
+			SNew(SButton)
+			.VAlign(VAlign_Center)
+			.IsEnabled(this, &SMediaSourceOverlay::Rewind_IsEnabled)
+			.OnClicked(this, &SMediaSourceOverlay::Rewind_OnClicked)
+			[
+				SNew(SImage)
+				.ColorAndOpacity(FSlateColor::UseForeground())
+				.Image(FSlateIconFinder::FindIcon("MediaPlateEditor.RewindMedia.Small").GetIcon())
+				.ToolTipText(LOCTEXT("Rewind", "Rewind the media to the beginning"))
+			]
+		]
+
+		// Reverse button.
+		+ SHorizontalBox::Slot()
+		.VAlign(VAlign_Center)
+		.AutoWidth()
+		[
+			SNew(SButton)
+			.VAlign(VAlign_Center)
+			.IsEnabled(this, &SMediaSourceOverlay::Reverse_IsEnabled)
+			.OnClicked(this, &SMediaSourceOverlay::Reverse_OnClicked)
+			[
+				SNew(SImage)
+				.ColorAndOpacity(FSlateColor::UseForeground())
+				.Image(FSlateIconFinder::FindIcon("MediaPlateEditor.ReverseMedia.Small").GetIcon())
+				.ToolTipText(LOCTEXT("Reverse", "Play media in reverse."))
+			]
+		]
+
+		// Play button.
+		+ SHorizontalBox::Slot()
+		.VAlign(VAlign_Center)
+		.AutoWidth()
+		[
+			SNew(SButton)
+			.VAlign(VAlign_Center)
+			.IsEnabled(this, &SMediaSourceOverlay::Play_IsEnabled)
+			.Visibility(this, &SMediaSourceOverlay::Play_GetVisibility)
+			.OnClicked(this, &SMediaSourceOverlay::Play_OnClicked)
+			[
+				SNew(SImage)
+				.ColorAndOpacity(FSlateColor::UseForeground())
+				.Image(FSlateIconFinder::FindIcon("MediaPlateEditor.PlayMedia.Small").GetIcon())
+				.ToolTipText(LOCTEXT("Play", "Start media playback"))
+			]
+		]
+
+		// Pause button.
+		+ SHorizontalBox::Slot()
+		.VAlign(VAlign_Center)
+		.AutoWidth()
+		[
+			SNew(SButton)
+			.VAlign(VAlign_Center)
+			.IsEnabled(this, &SMediaSourceOverlay::Pause_IsEnabled)
+			.Visibility(this, &SMediaSourceOverlay::Pause_GetVisibility)
+			.OnClicked(this, &SMediaSourceOverlay::Pause_OnClicked)
+			[
+				SNew(SImage)
+				.ColorAndOpacity(FSlateColor::UseForeground())
+				.Image(FSlateIconFinder::FindIcon("MediaPlateEditor.PauseMedia.Small").GetIcon())
+				.ToolTipText(LOCTEXT("Pause", "Pause media playback"))
+			]
+		]
+
+		// Forward button.
+		+ SHorizontalBox::Slot()
+		.VAlign(VAlign_Center)
+		.AutoWidth()
+		[
+			SNew(SButton)
+			.VAlign(VAlign_Center)
+			.IsEnabled(this, &SMediaSourceOverlay::Forward_IsEnabled)
+			.OnClicked(this, &SMediaSourceOverlay::Forward_OnClicked)
+			[
+				SNew(SImage)
+				.ColorAndOpacity(FSlateColor::UseForeground())
+				.Image(FSlateIconFinder::FindIcon("MediaPlateEditor.RewindMedia.Small").GetIcon())
+				// There is no correct icon pointing forward, so this rotates the rewind icon 180 degrees and adds the required offset.
+				.RenderTransform(FSlateRenderTransform(TTransform2<float>(TQuat2<float>(PI), FVector2f(20.f, 20.f))))
+				.ToolTipText(LOCTEXT("Forward", "Fast forward the media to the end."))
+			]
+		]
+
+		// Current time
+		+ SHorizontalBox::Slot()
+		.VAlign(VAlign_Center)
+		.AutoWidth()
+		.Padding(10.f, 0.f, 0.f, 0.f)
+		[
+			SNew(STextBlock)
+			.ColorAndOpacity(FStyleColors::Foreground.GetSpecifiedColor())
+			.ShadowColorAndOpacity(FStyleColors::Panel.GetSpecifiedColor())
+			.ShadowOffset(FVector2D(1.f, 1.f))
+			.Font(IDetailLayoutBuilder::GetDetailFont())
+			.Text(this, &SMediaSourceOverlay::GetCurrentTime)
+		]
+
+		+ SHorizontalBox::Slot()
+		.VAlign(VAlign_Center)
+		.AutoWidth()
+		.Padding(5.f, 0.f, 0.f, 0.f)
+		[
+			SNew(STextBlock)
+			.ColorAndOpacity(FStyleColors::Foreground.GetSpecifiedColor())
+			.ShadowColorAndOpacity(FStyleColors::Panel.GetSpecifiedColor())
+			.ShadowOffset(FVector2D(1.f, 1.f))
+			.Font(IDetailLayoutBuilder::GetDetailFont())
+			.Text(INVTEXT("/"))
+		]
+
+		// Total time
+		+ SHorizontalBox::Slot()
+		.HAlign(HAlign_Left)
+		.VAlign(VAlign_Center)
+		.Padding(5.f, 0.f, 0.f, 0.f)
+		.FillWidth(1.f)
+		[
+			SNew(STextBlock)
+			.ColorAndOpacity(FStyleColors::Foreground.GetSpecifiedColor())
+			.ShadowColorAndOpacity(FStyleColors::Panel.GetSpecifiedColor())
+			.ShadowOffset(FVector2D(1.f, 1.f))
+			.Font(IDetailLayoutBuilder::GetDetailFont())
+			.Text(this, &SMediaSourceOverlay::GetTotalTime)
+		];
+}
+
+FText SMediaSourceOverlay::GetCurrentFrame() const
+{
+	if (UMediaPlayer* MediaPlayer = GetMediaPlayer())
+	{
+		if (!FrameRateFloat.IsSet())
+		{
+			const_cast<SMediaSourceOverlay*>(this)->TrySetFrameRate();
+		}
+
+		if (FrameRateFloat.IsSet())
+		{
+			return FText::AsNumber(FMath::FloorToInt(MediaPlayer->GetTime().GetTotalSeconds() * static_cast<double>(FrameRateFloat.GetValue())) + 1);
+		}
+	}
+
+	return INVTEXT("-");
+}
+
+FText SMediaSourceOverlay::GetTotalFrames() const
+{
+	if (!TotalFrames.IsSet())
+	{
+		if (UMediaPlayer* MediaPlayer = GetMediaPlayer())
+		{
+			if (!FrameRateFloat.IsSet())
+			{
+				const_cast<SMediaSourceOverlay*>(this)->TrySetFrameRate();
+			}
+
+			if (FrameRateFloat.IsSet())
+			{
+				const FTimespan Duration = MediaPlayer->GetDuration();
+
+				TotalFrames = FText::AsNumber(FMath::FloorToInt(Duration.GetTotalSeconds() * static_cast<double>(FrameRateFloat.GetValue())));
+			}
+		}
+	}
+
+	return TotalFrames.Get(INVTEXT("-"));
+}
+
+FText SMediaSourceOverlay::GetCurrentTime() const
+{
+	if (UMediaPlayer* MediaPlayer = GetMediaPlayer())
+	{
+		if (!FrameRateFloat.IsSet())
+		{
+			const_cast<SMediaSourceOverlay*>(this)->TrySetFrameRate();
+		}
+
+		if (FrameRateFloat.IsSet())
+		{
+			// No straight conversion from float to FFrameRate, so estimate like this.
+			FFrameRate FrameRate(FMath::RoundToInt(1000000.f * FrameRateFloat.GetValue()), 1000000);
+			FTimecode Timecode(MediaPlayer->GetTime().GetTotalSeconds(), FrameRate, /* DropFrame */ false, /* Rollover */ false);
+
+			return FText::FromString(Timecode.ToString());
+		}
+	}
+
+	return INVTEXT("-");
+}
+
+FText SMediaSourceOverlay::GetTotalTime() const
+{
+	if (!TotalTime.IsSet())
+	{
+		if (UMediaPlayer* MediaPlayer = GetMediaPlayer())
+		{
+			if (!FrameRateFloat.IsSet())
+			{
+				const_cast<SMediaSourceOverlay*>(this)->TrySetFrameRate();
+			}
+
+			if (FrameRateFloat.IsSet())
+			{
+				// No straight conversion from float to FFrameRate, so estimate like this.
+				const FFrameRate FrameRate(FMath::RoundToInt(1000000.f * FrameRateFloat.GetValue()), 1000000);
+				const FTimecode Timecode(MediaPlayer->GetDuration().GetTotalSeconds(), FrameRate, /* DropFrame */ false, /* Rollover */ false);
+
+				TotalTime = FText::FromString(Timecode.ToString());
+			}
+		}
+	}
+
+	return TotalTime.Get(INVTEXT("-"));
+}
+
+bool SMediaSourceOverlay::Rewind_IsEnabled() const
+{
+	if (UMediaPlayer* MediaPlayer = GetMediaPlayer())
+	{
+		return MediaPlayer->IsReady() &&
+			MediaPlayer->SupportsSeeking() &&
+			MediaPlayer->GetTime() > FTimespan::Zero();
+	}
+
+	return false;
+}
+
+FReply SMediaSourceOverlay::Rewind_OnClicked()
+{
+	if (UMediaPlayer* MediaPlayer = GetMediaPlayer())
+	{
+		MediaPlayer->Pause();
+		MediaPlayer->Rewind();
+	}
+
+	return FReply::Handled();
+}
+
+bool SMediaSourceOverlay::Reverse_IsEnabled() const
+{
+	if (UMediaPlayer* MediaPlayer = GetMediaPlayer())
+	{
+		const float Rate = FMath::Abs(MediaPlayer->GetRate());
+
+		return MediaPlayer->IsReady() && MediaPlayer->SupportsRate(Rate > 0 ? -Rate : -1.f, /* Unthinned */ true);
+	}
+
+	return false;
+}
+
+FReply SMediaSourceOverlay::Reverse_OnClicked()
+{
+	if (UMediaPlayer* MediaPlayer = GetMediaPlayer())
+	{
+		const float Rate = FMath::Abs(MediaPlayer->GetRate());
+
+		MediaPlayer->SetRate(Rate > 0 ? -Rate : -1.f);
+	}
+
+	return FReply::Handled();
+}
+
+bool SMediaSourceOverlay::Play_IsEnabled() const
+{
+	if (UMediaPlayer* MediaPlayer = GetMediaPlayer())
+	{
+		return MediaPlayer->IsReady()
+			&& (!MediaPlayer->IsPlaying() || (MediaPlayer->GetRate() != 1.0f));
+	}
+
+	return false;
+}
+
+EVisibility SMediaSourceOverlay::Play_GetVisibility() const
+{
+	if (UMediaPlayer* MediaPlayer = GetMediaPlayer())
+	{
+		return MediaPlayer->IsPlaying() 
+			? EVisibility::Collapsed 
+			: EVisibility::Visible;
+	}
+
+	return EVisibility::Visible;
+}
+
+FReply SMediaSourceOverlay::Play_OnClicked()
+{
+	if (UMediaPlayer* MediaPlayer = GetMediaPlayer())
+	{
+		MediaPlayer->Play();
+	}
+
+	return FReply::Handled();
+}
+
+bool SMediaSourceOverlay::Pause_IsEnabled() const
+{
+	if (UMediaPlayer* MediaPlayer = GetMediaPlayer())
+	{
+		return MediaPlayer->CanPause() && !MediaPlayer->IsPaused();
+	}
+
+	return false;
+}
+
+EVisibility SMediaSourceOverlay::Pause_GetVisibility() const
+{
+	if (UMediaPlayer* MediaPlayer = GetMediaPlayer())
+	{
+		const bool bIsVisible = MediaPlayer->CanPause() && !MediaPlayer->IsPaused();
+		
+		return bIsVisible 
+			? EVisibility::Visible 
+			: EVisibility::Collapsed;
+	}
+
+	return EVisibility::Visible;
+}
+
+FReply SMediaSourceOverlay::Pause_OnClicked()
+{
+	if (UMediaPlayer* MediaPlayer = GetMediaPlayer())
+	{
+		MediaPlayer->Pause();
+	}
+
+	return FReply::Handled();
+}
+
+bool SMediaSourceOverlay::Forward_IsEnabled() const
+{
+	if (UMediaPlayer* MediaPlayer = GetMediaPlayer())
+	{
+		return MediaPlayer->IsReady();
+	}
+
+	return false;
+}
+
+FReply SMediaSourceOverlay::Forward_OnClicked()
+{
+	if (UMediaPlayer* MediaPlayer = GetMediaPlayer())
+	{
+		// IDEX_NONE here access the currently playing TrackIndex and FormatIndex.
+		const double FrameRate = MediaPlayer->GetVideoTrackFrameRate(INDEX_NONE, INDEX_NONE);
+
+		if (!FMath::IsNearlyZero(FrameRate))
+		{
+			MediaPlayer->Pause();
+
+			const double FrameTime = 1.0 / FrameRate;
+			FTimespan SeekLocation = MediaPlayer->GetDuration();
+			SeekLocation -= FTimespan::FromSeconds(FrameTime);
+
+			MediaPlayer->Seek(SeekLocation);
+		}
+	}
+
+	return FReply::Handled();
+}
+
+}
+
+#undef LOCTEXT_NAMESPACE
diff --git a/Engine/Plugins/Experimental/MediaViewer/Source/MediaViewer/Private/Widgets/SMediaSourceOverlay.h b/Engine/Plugins/Experimental/MediaViewer/Source/MediaViewer/Private/Widgets/SMediaSourceOverlay.h
new file mode 100644
--- /dev/null
+++ b/Engine/Plugins/Experimental/MediaViewer/Source/MediaViewer/Private/Widgets/SMediaSourceOverlay.h
@@ -0,0 +1,74 @@
+// Copyright Epic Games, Inc. All Rights Reserved.
+
+#pragma once
+
+#include "Widgets/SCompoundWidget.h"
+
+#include "Internationalization/Text.h"
+#include "Misc/Optional.h"
+
+class IMediaStreamPlayer;
+class UMediaPlayer;
+class UMediaStream;
+
+namespace UE::MediaViewer::Private
+{
+
+class FMediaSourceImageViewer;
+	
+class SMediaSourceOverlay : public SCompoundWidget
+{
+	SLATE_DECLARE_WIDGET(SMediaSourceOverlay, SCompoundWidget)
+
+	SLATE_BEGIN_ARGS(SMediaSourceOverlay)
+		{}
+	SLATE_END_ARGS()
+
+public:
+	virtual ~SMediaSourceOverlay() override = default;
+
+	void Construct(const FArguments& InArgs, const TSharedRef<FMediaSourceImageViewer>& InImageViewer);
+
+protected:
+	TWeakPtr<FMediaSourceImageViewer> ImageViewerWeak;
+	TOptional<float> FrameRateFloat;
+	mutable TOptional<FText> TotalFrames;
+	mutable TOptional<FText> TotalTime;
+
+	UMediaStream* GetMediaStream() const;
+
+	IMediaStreamPlayer* GetMediaStreamPlayer() const;
+
+	UMediaPlayer* GetMediaPlayer() const;
+
+	void TrySetFrameRate();
+
+	TSharedRef<SWidget> CreateSlider();
+
+	TSharedRef<SWidget> CreateControls();
+
+	FText GetCurrentFrame() const;
+	FText GetTotalFrames() const;
+
+	FText GetCurrentTime() const;
+	FText GetTotalTime() const;
+
+	bool Rewind_IsEnabled() const;
+	FReply Rewind_OnClicked();
+
+	bool Reverse_IsEnabled() const;
+	FReply Reverse_OnClicked();
+
+	bool Play_IsEnabled() const;
+	EVisibility Play_GetVisibility() const;
+	FReply Play_OnClicked();
+
+	bool Pause_IsEnabled() const;
+	EVisibility Pause_GetVisibility() const;
+	FReply Pause_OnClicked();
+
+	bool Forward_IsEnabled() const;
+	FReply Forward_OnClicked();
+};
+
+} // UE::MediaViewer::Private
\ No newline at end of file
diff --git a/Engine/Plugins/Experimental/MediaViewer/Source/MediaViewer/Private/Widgets/SMediaViewerLibraryItem.cpp b/Engine/Plugins/Experimental/MediaViewer/Source/MediaViewer/Private/Widgets/SMediaViewerLibraryItem.cpp
--- a/Engine/Plugins/Experimental/MediaViewer/Source/MediaViewer/Private/Widgets/SMediaViewerLibraryItem.cpp
+++ b/Engine/Plugins/Experimental/MediaViewer/Source/MediaViewer/Private/Widgets/SMediaViewerLibraryItem.cpp
@@ -65,8 +65,8 @@ void SMediaViewerLibraryItem::Construct(const FArguments& InArgs, const TSharedR
 		.Padding(0.f, 0.f, 5.f, 0.f)
 		[
 			SNew(STextBlock)
-				.Text(Item->GetItemTypeDisplayName())
-				.Font(IDetailLayoutBuilder::GetDetailFont())
+			.Font(IDetailLayoutBuilder::GetDetailFont())
+			.Text(Item->GetItemTypeDisplayName())
 		];
 
 	if (!Item->IsTransient())
