Media Viewer: Added RHI Texture sampling
- Added the media source RHI texture sampling to the other image viewers, vastly reducing overhead in those cases.

diff --git a/Engine/Plugins/Experimental/MediaViewer/Source/MediaViewer/Private/ImageViewers/MaterialInterfaceImageViewer.cpp b/Engine/Plugins/Experimental/MediaViewer/Source/MediaViewer/Private/ImageViewers/MaterialInterfaceImageViewer.cpp
--- a/Engine/Plugins/Experimental/MediaViewer/Source/MediaViewer/Private/ImageViewers/MaterialInterfaceImageViewer.cpp
+++ b/Engine/Plugins/Experimental/MediaViewer/Source/MediaViewer/Private/ImageViewers/MaterialInterfaceImageViewer.cpp
@@ -189,7 +189,7 @@ TSharedPtr<FMediaImageViewer> FMaterialInterfaceImageViewer::FItem::CreateImageV
 
 FString FMaterialInterfaceImageViewer::FItem::GetReferencerName() const
 {
-	static const FString ReferencerName = TEXT("FMaterialInterfaceImageViewer");
+	static const FString ReferencerName = TEXT("FMaterialInterfaceImageViewer::FItem");
 	return ReferencerName;
 }
 
@@ -251,34 +251,25 @@ TOptional<TVariant<FColor, FLinearColor>> FMaterialInterfaceImageViewer::GetPixe
 {
 	if (InPixelCoords.X < 0 || InPixelCoords.Y < 0)
 	{
+		SampleCache.Invalidate();
 		return {};
 	}
 
-	UTextureRenderTarget2D* RenderTarget = MaterialSettings.RenderTarget;
-
-	if (!RenderTarget)
+	if (InPixelCoords.X >= ImageInfo.Size.X || InPixelCoords.Y >= ImageInfo.Size.Y)
 	{
+		SampleCache.Invalidate();
 		return {};
 	}
 
-	const FIntPoint RenderTargetSize = {static_cast<int32>(RenderTarget->Source.GetSizeX()), static_cast<int32>(RenderTarget->Source.GetSizeY())};
-
-	if (InPixelCoords.X >= RenderTargetSize.X || InPixelCoords.Y >= RenderTargetSize.Y)
+	if (const FLinearColor* PixelColor = SampleCache.GetPixelColor(InPixelCoords))
 	{
-		return {};
-	}
-
-	const FImage* CachedMip = MipCache.GetMipImage(InMipLevel);
+		TVariant<FColor, FLinearColor> PixelColorVariant;
+		PixelColorVariant.Set<FLinearColor>(*PixelColor);
 
-	if (!CachedMip)
-	{
-		return {};
+		return PixelColorVariant;
 	}
 
-	TVariant<FColor, FLinearColor> Pixel;
-	Pixel.Set<FLinearColor>(CachedMip->GetOnePixelLinear(InPixelCoords.X, InPixelCoords.Y));
-
-	return Pixel;
+	return {};
 }
 
 TSharedPtr<FStructOnScope> FMaterialInterfaceImageViewer::GetCustomSettingsOnScope() const
@@ -331,7 +322,7 @@ void FMaterialInterfaceImageViewer::CreateBrush()
 	UTextureRenderTarget2D* RenderTarget = FMediaViewerUtils::RenderMaterial(MaterialSettings.MaterialInterface);
 	MaterialSettings.RenderTarget = RenderTarget;
 
-	MipCache = FTextureMipCache(RenderTarget);
+	SampleCache = FTextureSampleCache(RenderTarget);
 
 	Brush = MakeShared<FSlateImageBrush>(RenderTarget, FVector2D(RenderTarget->GetSurfaceWidth(), RenderTarget->GetSurfaceDepth()));
 }
@@ -341,7 +332,7 @@ void FMaterialInterfaceImageViewer::RenderMaterial()
 	if (::IsValid(MaterialSettings.MaterialInterface) && ::IsValid(MaterialSettings.RenderTarget))
 	{
 		FMediaViewerUtils::RenderMaterial(MaterialSettings.MaterialInterface, MaterialSettings.RenderTarget);
-		MipCache.Invalidate();
+		SampleCache.MarkDirty();
 	}
 }
 
diff --git a/Engine/Plugins/Experimental/MediaViewer/Source/MediaViewer/Private/ImageViewers/MaterialInterfaceImageViewer.h b/Engine/Plugins/Experimental/MediaViewer/Source/MediaViewer/Private/ImageViewers/MaterialInterfaceImageViewer.h
--- a/Engine/Plugins/Experimental/MediaViewer/Source/MediaViewer/Private/ImageViewers/MaterialInterfaceImageViewer.h
+++ b/Engine/Plugins/Experimental/MediaViewer/Source/MediaViewer/Private/ImageViewers/MaterialInterfaceImageViewer.h
@@ -7,7 +7,7 @@
 #include "Library/MediaViewerLibraryItem.h"
 #include "UObject/GCObject.h"
 
-#include "ImageViewers/TextureMipCache.h"
+#include "ImageViewers/TextureSampleCache.h"
 #include "Templates/SharedPointer.h"
 
 #include "MaterialInterfaceImageViewer.generated.h"
@@ -80,7 +80,7 @@ public:
 		virtual TSharedPtr<FSlateBrush> CreateThumbnail() override;
 		virtual TSharedPtr<FMediaImageViewer> CreateImageViewer() const override;
 		virtual TSharedPtr<FMediaViewerLibraryItem> Clone() const override;
-		//~ End FMediaViewerLibraryItem	
+		//~ End FMediaViewerLibraryItem
 
 		//~ Begin FGCObject
 		virtual FString GetReferencerName() const override;
@@ -113,7 +113,7 @@ public:
 
 protected:
 	FMaterialInterfaceImageViewerSettings MaterialSettings;
-	FTextureMipCache MipCache;
+	mutable FTextureSampleCache SampleCache;
 
 	void CreateBrush();
 
diff --git a/Engine/Plugins/Experimental/MediaViewer/Source/MediaViewer/Private/ImageViewers/MediaTextureImageViewer.cpp b/Engine/Plugins/Experimental/MediaViewer/Source/MediaViewer/Private/ImageViewers/MediaTextureImageViewer.cpp
--- a/Engine/Plugins/Experimental/MediaViewer/Source/MediaViewer/Private/ImageViewers/MediaTextureImageViewer.cpp
+++ b/Engine/Plugins/Experimental/MediaViewer/Source/MediaViewer/Private/ImageViewers/MediaTextureImageViewer.cpp
@@ -172,32 +172,17 @@ FMediaTextureImageViewer::FMediaTextureImageViewer(const FGuid& InId, TNotNull<U
 {
 	MediaTextureSettings.MediaTexture = InMediaTexture;
 
-	const FIntPoint Size = {
-		FMath::RoundToInt(InMediaTexture->GetSurfaceWidth()),
-		FMath::RoundToInt(InMediaTexture->GetSurfaceHeight())
-	};
-
-	if (Size.X < 2 || Size.Y < 2)
+	if (ImageInfo.Size.X < 2 || ImageInfo.Size.Y < 2)
 	{
 		return;
 	}
-	
-	MediaTextureSettings.TextureRenderMaterial = UMaterialInstanceDynamic::Create(
-		FMediaViewerUtils::GetTextureRenderMaterial(),
-		GetTransientPackage()
-	);
-
-	MediaTextureSettings.TextureRenderMaterial->SetTextureParameterValue(
-		FMediaViewerUtils::TextureRenderMaterial_TextureParameterName,
-		InMediaTexture
-	);
-
-	MediaTextureSettings.RenderTarget = FMediaViewerUtils::RenderMaterial(MediaTextureSettings.TextureRenderMaterial);
 
 	Brush = MakeShared<FSlateImageBrush>(
 		InMediaTexture,
-		FVector2D(InMediaTexture->GetSurfaceWidth(), InMediaTexture->GetSurfaceHeight())
+		FVector2D(ImageInfo.Size.X, ImageInfo.Size.Y)
 	);
+
+	SampleCache = FTextureSampleCache(InMediaTexture);
 }
 
 TSharedPtr<FMediaViewerLibraryItem> FMediaTextureImageViewer::CreateLibraryItem() const
@@ -220,33 +205,30 @@ TSharedPtr<FMediaViewerLibraryItem> FMediaTextureImageViewer::CreateLibraryItem(
 
 TOptional<TVariant<FColor, FLinearColor>> FMediaTextureImageViewer::GetPixelColor(const FIntPoint& InPixelCoords, int32 InMipLevel) const
 {
+	// Always mark the sample cache dirty. We have no control over the state of the media texture.
+	SampleCache.MarkDirty();
+
 	if (InPixelCoords.X < 0 || InPixelCoords.Y < 0)
 	{
+		SampleCache.Invalidate();
 		return {};
 	}
 
-	UTextureRenderTarget2D* RenderTarget = MediaTextureSettings.RenderTarget;
-
-	if (!RenderTarget)
+	if (InPixelCoords.X >= ImageInfo.Size.X || InPixelCoords.Y >= ImageInfo.Size.Y)
 	{
+		SampleCache.Invalidate();
 		return {};
 	}
 
-	const FIntPoint RenderTargetSize = {static_cast<int32>(RenderTarget->Source.GetSizeX()), static_cast<int32>(RenderTarget->Source.GetSizeY())};
-
-	if (InPixelCoords.X >= RenderTargetSize.X || InPixelCoords.Y >= RenderTargetSize.Y)
+	if (const FLinearColor* PixelColor = SampleCache.GetPixelColor(InPixelCoords))
 	{
-		return {};
-	}
-
-	TArray64<uint8> MipData;
+		TVariant<FColor, FLinearColor> PixelColorVariant;
+		PixelColorVariant.Set<FLinearColor>(*PixelColor);
 
-	if (!RenderTarget->Source.GetMipData(MipData, InMipLevel))
-	{
-		return {};
+		return PixelColorVariant;
 	}
 
-	return FMediaViewerUtils::GetPixelColor(MipData, RenderTarget->GetFormat(), RenderTargetSize, InPixelCoords, InMipLevel);
+	return {};
 }
 
 TSharedPtr<FStructOnScope> FMediaTextureImageViewer::GetCustomSettingsOnScope() const
@@ -271,30 +253,6 @@ FString FMediaTextureImageViewer::GetReferencerName() const
 	return ReferencerName;
 }
 
-void FMediaTextureImageViewer::PaintImage(FMediaImagePaintParams& InPaintParams, const FMediaImagePaintGeometry& InPaintGeometry)
-{
-	if (MediaTextureSettings.bRealTime)
-	{
-		if (UMediaTexture* MediaTexture = MediaTextureSettings.MediaTexture)
-		{
-			const FIntPoint Size = {
-				FMath::RoundToInt(MediaTexture->GetSurfaceWidth()),
-				FMath::RoundToInt(MediaTexture->GetSurfaceHeight())
-			};
-
-			if (Size.X > 1 && Size.Y > 1)
-			{
-				FMediaViewerUtils::RenderMaterial(
-					MediaTextureSettings.TextureRenderMaterial,
-					MediaTextureSettings.RenderTarget
-				);
-			}
-		}
-	}
-
-	FMediaImageViewer::PaintImage(InPaintParams, InPaintGeometry);
-}
-
 } // UE::MediaViewer::Private
 
 #undef LOCTEXT_NAMESPACE
diff --git a/Engine/Plugins/Experimental/MediaViewer/Source/MediaViewer/Private/ImageViewers/MediaTextureImageViewer.h b/Engine/Plugins/Experimental/MediaViewer/Source/MediaViewer/Private/ImageViewers/MediaTextureImageViewer.h
--- a/Engine/Plugins/Experimental/MediaViewer/Source/MediaViewer/Private/ImageViewers/MediaTextureImageViewer.h
+++ b/Engine/Plugins/Experimental/MediaViewer/Source/MediaViewer/Private/ImageViewers/MediaTextureImageViewer.h
@@ -7,13 +7,12 @@
 #include "Library/MediaViewerLibraryItem.h"
 #include "UObject/GCObject.h"
 
+#include "ImageViewers/TextureSampleCache.h"
 #include "Templates/SharedPointer.h"
 
 #include "MediaTextureImageViewer.generated.h"
 
-class UMaterialInstanceDynamic;
 class UMediaTexture;
-class UTextureRenderTarget2D;
 
 USTRUCT()
 struct FMediaTextureImageViewerSettings
@@ -25,12 +24,6 @@ struct FMediaTextureImageViewerSettings
 
 	UPROPERTY(VisibleAnywhere, Category = "Texture")
 	TObjectPtr<UMediaTexture> MediaTexture;
-
-	UPROPERTY(VisibleAnywhere, Category = "Material")
-	TObjectPtr<UTextureRenderTarget2D> RenderTarget;
-
-	UPROPERTY(VisibleAnywhere, Category = "Material")
-	TObjectPtr<UMaterialInstanceDynamic> TextureRenderMaterial;
 };
 
 namespace UE::MediaViewer::Private
@@ -97,10 +90,7 @@ public:
 
 protected:
 	FMediaTextureImageViewerSettings MediaTextureSettings;
-
-	//~ Begin FMediaImageViewer
-	virtual void PaintImage(FMediaImagePaintParams& InPaintParams, const FMediaImagePaintGeometry& InPaintGeometry) override;
-	//~ End FMediaImageViewer
+	mutable FTextureSampleCache SampleCache;
 };
 
 } // UE::MediaViewer::Private
diff --git a/Engine/Plugins/Experimental/MediaViewer/Source/MediaViewer/Private/ImageViewers/Texture2DImageViewer.cpp b/Engine/Plugins/Experimental/MediaViewer/Source/MediaViewer/Private/ImageViewers/Texture2DImageViewer.cpp
--- a/Engine/Plugins/Experimental/MediaViewer/Source/MediaViewer/Private/ImageViewers/Texture2DImageViewer.cpp
+++ b/Engine/Plugins/Experimental/MediaViewer/Source/MediaViewer/Private/ImageViewers/Texture2DImageViewer.cpp
@@ -173,7 +173,7 @@ TSharedPtr<FMediaViewerLibraryItem> FTexture2DImageViewer::FItem::Clone() const
 
 FString FTexture2DImageViewer::FItem::GetReferencerName() const
 {
-	static const FString ReferencerName = TEXT("FTexture2DImageViewer");
+	static const FString ReferencerName = TEXT("FTexture2DImageViewer::FItem");
 	return ReferencerName;
 }
 
@@ -197,7 +197,7 @@ FTexture2DImageViewer::FTexture2DImageViewer(const FGuid& InId, TNotNull<UTextur
 		InTexture->GetPlatformData()->Mips.Num(),
 		FText::FromName(InTexture->GetFName())
 	})
-	, MipCache(InTexture)
+	, SampleCache(InTexture)
 {
 	TextureSettings.Texture = InTexture;
 
@@ -232,31 +232,20 @@ TOptional<TVariant<FColor, FLinearColor>> FTexture2DImageViewer::GetPixelColor(c
 		return {};
 	}
 
-	UTexture2D* Texture = TextureSettings.Texture;
-
-	if (!Texture)
+	if (InPixelCoords.X >= ImageInfo.Size.X || InPixelCoords.Y >= ImageInfo.Size.Y)
 	{
 		return {};
 	}
 
-	const FIntPoint TextureSize = {static_cast<int32>(Texture->Source.GetSizeX()), static_cast<int32>(Texture->Source.GetSizeY())};
-
-	if (InPixelCoords.X >= TextureSize.X || InPixelCoords.Y >= TextureSize.Y)
+	if (const FLinearColor* PixelColor = SampleCache.GetPixelColor(InPixelCoords))
 	{
-		return {};
-	}
-
-	const FImage* CachedMip = MipCache.GetMipImage(InMipLevel);
+		TVariant<FColor, FLinearColor> PixelColorVariant;
+		PixelColorVariant.Set<FLinearColor>(*PixelColor);
 
-	if (!CachedMip)
-	{
-		return {};
+		return PixelColorVariant;
 	}
 
-	TVariant<FColor, FLinearColor> Pixel;
-	Pixel.Set<FLinearColor>(CachedMip->GetOnePixelLinear(InPixelCoords.X, InPixelCoords.Y));
-
-	return Pixel;
+	return {};
 }
 
 TSharedPtr<FStructOnScope> FTexture2DImageViewer::GetCustomSettingsOnScope() const
diff --git a/Engine/Plugins/Experimental/MediaViewer/Source/MediaViewer/Private/ImageViewers/Texture2DImageViewer.h b/Engine/Plugins/Experimental/MediaViewer/Source/MediaViewer/Private/ImageViewers/Texture2DImageViewer.h
--- a/Engine/Plugins/Experimental/MediaViewer/Source/MediaViewer/Private/ImageViewers/Texture2DImageViewer.h
+++ b/Engine/Plugins/Experimental/MediaViewer/Source/MediaViewer/Private/ImageViewers/Texture2DImageViewer.h
@@ -7,7 +7,7 @@
 #include "Library/MediaViewerLibraryItem.h"
 #include "UObject/GCObject.h"
 
-#include "ImageViewers/TextureMipCache.h"
+#include "ImageViewers/TextureSampleCache.h"
 #include "Templates/SharedPointer.h"
 
 #include "Texture2DImageViewer.generated.h"
@@ -94,7 +94,7 @@ public:
 
 protected:
 	FTexture2DImageViewerSettings TextureSettings;
-	FTextureMipCache MipCache;
+	mutable FTextureSampleCache SampleCache;
 };
 
 } // UE::MediaViewer::Private
diff --git a/Engine/Plugins/Experimental/MediaViewer/Source/MediaViewer/Private/ImageViewers/TextureRenderTarget2DImageViewer.cpp b/Engine/Plugins/Experimental/MediaViewer/Source/MediaViewer/Private/ImageViewers/TextureRenderTarget2DImageViewer.cpp
--- a/Engine/Plugins/Experimental/MediaViewer/Source/MediaViewer/Private/ImageViewers/TextureRenderTarget2DImageViewer.cpp
+++ b/Engine/Plugins/Experimental/MediaViewer/Source/MediaViewer/Private/ImageViewers/TextureRenderTarget2DImageViewer.cpp
@@ -173,6 +173,8 @@ FTextureRenderTarget2DImageViewer::FTextureRenderTarget2DImageViewer(const FGuid
 		InRenderTarget,
 		FVector2D(InRenderTarget->GetSurfaceWidth(), InRenderTarget->GetSurfaceHeight())
 	);
+	
+	SampleCache = FTextureSampleCache(InRenderTarget);
 }
 
 TSharedPtr<FMediaViewerLibraryItem> FTextureRenderTarget2DImageViewer::CreateLibraryItem() const
@@ -195,39 +197,30 @@ TSharedPtr<FMediaViewerLibraryItem> FTextureRenderTarget2DImageViewer::CreateLib
 
 TOptional<TVariant<FColor, FLinearColor>> FTextureRenderTarget2DImageViewer::GetPixelColor(const FIntPoint& InPixelCoords, int32 InMipLevel) const
 {
+	// Always mark the sample cache dirty. We have no control over the state of the media texture.
+	SampleCache.MarkDirty();
+
 	if (InPixelCoords.X < 0 || InPixelCoords.Y < 0)
 	{
+		SampleCache.Invalidate();
 		return {};
 	}
 
-	UTextureRenderTarget2D* RenderTarget = RenderTargetSettings.RenderTarget;
-
-	if (!RenderTarget)
+	if (InPixelCoords.X >= ImageInfo.Size.X || InPixelCoords.Y >= ImageInfo.Size.Y)
 	{
+		SampleCache.Invalidate();
 		return {};
 	}
 
-	const FIntPoint RenderTargetSize = {static_cast<int32>(RenderTarget->Source.GetSizeX()), static_cast<int32>(RenderTarget->Source.GetSizeY())};
-
-	if (InPixelCoords.X >= RenderTargetSize.X || InPixelCoords.Y >= RenderTargetSize.Y)
+	if (const FLinearColor* PixelColor = SampleCache.GetPixelColor(InPixelCoords))
 	{
-		return {};
-	}
-
-	TArray64<uint8> MipData;
+		TVariant<FColor, FLinearColor> PixelColorVariant;
+		PixelColorVariant.Set<FLinearColor>(*PixelColor);
 
-	if (!RenderTarget->Source.GetMipData(MipData, InMipLevel))
-	{
-		return {};
+		return PixelColorVariant;
 	}
 
-	return FMediaViewerUtils::GetPixelColor(
-		MipData, 
-		GetPixelFormatFromRenderTargetFormat(RenderTarget->RenderTargetFormat), 
-		RenderTargetSize,
-		InPixelCoords, 
-		InMipLevel
-	);
+	return {};
 }
 
 TSharedPtr<FStructOnScope> FTextureRenderTarget2DImageViewer::GetCustomSettingsOnScope() const
diff --git a/Engine/Plugins/Experimental/MediaViewer/Source/MediaViewer/Private/ImageViewers/TextureRenderTarget2DImageViewer.h b/Engine/Plugins/Experimental/MediaViewer/Source/MediaViewer/Private/ImageViewers/TextureRenderTarget2DImageViewer.h
--- a/Engine/Plugins/Experimental/MediaViewer/Source/MediaViewer/Private/ImageViewers/TextureRenderTarget2DImageViewer.h
+++ b/Engine/Plugins/Experimental/MediaViewer/Source/MediaViewer/Private/ImageViewers/TextureRenderTarget2DImageViewer.h
@@ -7,6 +7,7 @@
 #include "Library/MediaViewerLibraryItem.h"
 #include "UObject/GCObject.h"
 
+#include "ImageViewers/TextureSampleCache.h"
 #include "Templates/SharedPointer.h"
 
 #include "TextureRenderTarget2DImageViewer.generated.h"
@@ -86,6 +87,7 @@ public:
 
 protected:
 	FTextureRenderTarget2DImageViewerSettings RenderTargetSettings;
+	mutable FTextureSampleCache SampleCache;
 };
 
 } // UE::MediaViewer::Private
diff --git a/Engine/Plugins/Experimental/MediaViewer/Source/MediaViewer/Private/ImageViewers/TextureSampleCache.cpp b/Engine/Plugins/Experimental/MediaViewer/Source/MediaViewer/Private/ImageViewers/TextureSampleCache.cpp
--- a/Engine/Plugins/Experimental/MediaViewer/Source/MediaViewer/Private/ImageViewers/TextureSampleCache.cpp
+++ b/Engine/Plugins/Experimental/MediaViewer/Source/MediaViewer/Private/ImageViewers/TextureSampleCache.cpp
@@ -10,7 +10,6 @@ namespace UE::MediaViewer::Private
 {
 
 FTextureSampleCache::FTextureSampleCache()
-	: Texture(nullptr)
 {
 }
 
@@ -26,7 +25,7 @@ const FLinearColor* FTextureSampleCache::GetPixelColor(const FIntPoint& InPixelC
 		return nullptr;
 	}
 
-	const bool bNeedsUpdate = !PixelColorSample.IsSet() || InPixelCoordinates != PixelColorSample->Coordinates
+	const bool bNeedsUpdate = bDirty || !PixelColorSample.IsSet() || InPixelCoordinates != PixelColorSample->Coordinates
 		|| (InTime.IsSet() && (PixelColorSample->Time != InTime.GetValue()));
 
 	if (bNeedsUpdate)
@@ -35,8 +34,6 @@ const FLinearColor* FTextureSampleCache::GetPixelColor(const FIntPoint& InPixelC
 		(
 			[this, InPixelCoordinates, InTime](FRHICommandListImmediate& RHICmdList)
 			{
-				PixelColorSample.Reset();
-
 				if (FTextureResource* TextureResource = Texture->GetResource())
 				{
 					const FIntVector Size = TextureResource->TextureRHI->GetSizeXYZ();
@@ -51,8 +48,12 @@ const FLinearColor* FTextureSampleCache::GetPixelColor(const FIntPoint& InPixelC
 							InTime,
 							Data[0]
 						};
+						bDirty = false;
+						return;
 					}
 				}
+
+				PixelColorSample.Reset();
 			}
 		);
 	}
@@ -65,6 +66,11 @@ const FLinearColor* FTextureSampleCache::GetPixelColor(const FIntPoint& InPixelC
 	return nullptr;
 }
 
+void FTextureSampleCache::MarkDirty()
+{
+	bDirty = true;
+}
+
 void FTextureSampleCache::Invalidate()
 {
 	PixelColorSample.Reset();
diff --git a/Engine/Plugins/Experimental/MediaViewer/Source/MediaViewer/Private/ImageViewers/TextureSampleCache.h b/Engine/Plugins/Experimental/MediaViewer/Source/MediaViewer/Private/ImageViewers/TextureSampleCache.h
--- a/Engine/Plugins/Experimental/MediaViewer/Source/MediaViewer/Private/ImageViewers/TextureSampleCache.h
+++ b/Engine/Plugins/Experimental/MediaViewer/Source/MediaViewer/Private/ImageViewers/TextureSampleCache.h
@@ -21,6 +21,8 @@ public:
 
 	const FLinearColor* GetPixelColor(const FIntPoint& InPixelCoordinates, const TOptional<FTimespan>& InTime = {});
 
+	void MarkDirty();
+
 	void Invalidate();
 
 protected:
@@ -31,7 +33,8 @@ protected:
 		FLinearColor Color;
 	};
 
-	UTexture* Texture;
+	UTexture* Texture = nullptr;
+	bool bDirty = true;
 	TOptional<FPixelColorSample> PixelColorSample;
 };
 
diff --git a/Engine/Plugins/Experimental/MediaViewer/Source/MediaViewer/Private/ImageViewers/ViewportImageViewer.cpp b/Engine/Plugins/Experimental/MediaViewer/Source/MediaViewer/Private/ImageViewers/ViewportImageViewer.cpp
--- a/Engine/Plugins/Experimental/MediaViewer/Source/MediaViewer/Private/ImageViewers/ViewportImageViewer.cpp
+++ b/Engine/Plugins/Experimental/MediaViewer/Source/MediaViewer/Private/ImageViewers/ViewportImageViewer.cpp
@@ -35,34 +35,25 @@ TOptional<TVariant<FColor, FLinearColor>> FViewportImageViewer::GetPixelColor(co
 {
 	if (InPixelCoords.X < 0 || InPixelCoords.Y < 0)
 	{
+		SampleCache.Invalidate();
 		return {};
 	}
 
-	UTextureRenderTarget2D* RenderTarget = ViweportSettings.RenderTarget;
-
-	if (!RenderTarget)
+	if (InPixelCoords.X >= ImageInfo.Size.X || InPixelCoords.Y >= ImageInfo.Size.Y)
 	{
+		SampleCache.Invalidate();
 		return {};
 	}
 
-	const FIntPoint RenderTargetSize = {static_cast<int32>(RenderTarget->Source.GetSizeX()), static_cast<int32>(RenderTarget->Source.GetSizeY())};
-
-	if (InPixelCoords.X >= RenderTargetSize.X || InPixelCoords.Y >= RenderTargetSize.Y)
+	if (const FLinearColor* PixelColor = SampleCache.GetPixelColor(InPixelCoords))
 	{
-		return {};
-	}
-
-	const FImage* CachedMip = MipCache.GetMipImage(InMipLevel);
+		TVariant<FColor, FLinearColor> PixelColorVariant;
+		PixelColorVariant.Set<FLinearColor>(*PixelColor);
 
-	if (!CachedMip)
-	{
-		return {};
+		return PixelColorVariant;
 	}
 
-	TVariant<FColor, FLinearColor> Pixel;
-	Pixel.Set<FLinearColor>(CachedMip->GetOnePixelLinear(InPixelCoords.X, InPixelCoords.Y));
-
-	return Pixel;
+	return {};
 }
 
 TSharedPtr<FStructOnScope> FViewportImageViewer::GetCustomSettingsOnScope() const
@@ -111,9 +102,9 @@ void FViewportImageViewer::CreateBrush()
 	UTextureRenderTarget2D* RenderTarget = FMediaViewerUtils::CreateRenderTarget(ImageInfo.Size, /* Transparent */ false);
 	ViweportSettings.RenderTarget = RenderTarget;
 
-	MipCache = FTextureMipCache(RenderTarget);
-
 	Brush = MakeShared<FSlateImageBrush>(RenderTarget, FVector2D(RenderTarget->GetSurfaceWidth(), RenderTarget->GetSurfaceDepth()));
+
+	SampleCache = FTextureSampleCache(RenderTarget);
 }
 
 void FViewportImageViewer::RenderViewport(TNotNull<FViewport*> InViewport, TNotNull<UTextureRenderTarget2D*> InRenderTarget, 
@@ -133,7 +124,7 @@ void FViewportImageViewer::RenderViewport(TNotNull<FViewport*> InViewport, TNotN
 	ENQUEUE_RENDER_COMMAND(CopyViewportRenderTarget)(
 		[Viewport = InViewport, RenderTarget = InRenderTarget](FRHICommandListImmediate& InRHICmdList)
 		{
-			FTextureRenderTargetResource* RenderTargetResource = RenderTarget->GetRenderTargetResource();;
+			FTextureRenderTargetResource* RenderTargetResource = RenderTarget->GetRenderTargetResource();
 
 			if (!RenderTargetResource)
 			{
@@ -178,7 +169,7 @@ void FViewportImageViewer::RenderViewport()
 
 	RenderViewport(TNotNull<FViewport*>(Viewport), TNotNull<UTextureRenderTarget2D*>(ViweportSettings.RenderTarget));
 
-	MipCache.Invalidate();
+	SampleCache.MarkDirty();
 }
 
 } // UE::MediaViewer::Private
diff --git a/Engine/Plugins/Experimental/MediaViewer/Source/MediaViewer/Private/ImageViewers/ViewportImageViewer.h b/Engine/Plugins/Experimental/MediaViewer/Source/MediaViewer/Private/ImageViewers/ViewportImageViewer.h
--- a/Engine/Plugins/Experimental/MediaViewer/Source/MediaViewer/Private/ImageViewers/ViewportImageViewer.h
+++ b/Engine/Plugins/Experimental/MediaViewer/Source/MediaViewer/Private/ImageViewers/ViewportImageViewer.h
@@ -6,7 +6,7 @@
 #include "ImageViewer/MediaImageViewer.h"
 #include "UObject/GCObject.h"
 
-#include "ImageViewers/TextureMipCache.h"
+#include "ImageViewers/TextureSampleCache.h"
 #include "Templates/SharedPointer.h"
 
 #include "ViewportImageViewer.generated.h"
@@ -53,7 +53,7 @@ protected:
 	static void RenderViewport(TNotNull<FViewport*> InViewport, TNotNull<UTextureRenderTarget2D*> InRenderTarget, bool bResizeTargetToViewport = true);
 
 	FViewportImageViewerSettings ViweportSettings;
-	FTextureMipCache MipCache;
+	mutable FTextureSampleCache SampleCache;
 
 	virtual FViewport* GetViewport() const = 0;
 
