Media Viewer: Lock enabled by default

diff --git a/Engine/Plugins/Experimental/MediaViewer/Source/MediaViewer/Private/Widgets/MediaViewerSettings.h b/Engine/Plugins/Experimental/MediaViewer/Source/MediaViewer/Private/Widgets/MediaViewerSettings.h
--- a/Engine/Plugins/Experimental/MediaViewer/Source/MediaViewer/Private/Widgets/MediaViewerSettings.h
+++ b/Engine/Plugins/Experimental/MediaViewer/Source/MediaViewer/Private/Widgets/MediaViewerSettings.h
@@ -38,7 +38,7 @@ struct FMediaViewerSettings
 
 	/** If true, when one image transform changes, the other image transforms will change an equal amount. */
 	UPROPERTY(EditAnywhere, Category = "AB View", DisplayName = "Lock Transforms")
-	bool bAreTransformsLocked = false;
+	bool bAreTransformsLocked = true;
 
 	/** When in AB view, whether the viewers are split horizontally or vertically. */
 	TEnumAsByte<EOrientation> ABOrientation = Orient_Horizontal;
