Media Viewer: Library improvements.
- Renamed default group to "Saved".
- Dragging to other groups now works from the history group.

diff --git a/Engine/Plugins/Experimental/MediaViewer/MediaViewer.uplugin b/Engine/Plugins/Experimental/MediaViewer/MediaViewer.uplugin
--- a/Engine/Plugins/Experimental/MediaViewer/MediaViewer.uplugin
+++ b/Engine/Plugins/Experimental/MediaViewer/MediaViewer.uplugin
@@ -15,5 +15,15 @@
 			"Type": "Editor",
 			"LoadingPhase": "Default"
 		}
+	],
+	"Plugins": [
+		{
+			"Name": "MediaStream",
+			"Enabled": true
+		},
+		{
+			"Name": "MediaPlayerEditor",
+			"Enabled": true
+		}
 	]
 }
diff --git a/Engine/Plugins/Experimental/MediaViewer/Source/MediaViewer/Private/Library/MediaViewerLibrary.cpp b/Engine/Plugins/Experimental/MediaViewer/Source/MediaViewer/Private/Library/MediaViewerLibrary.cpp
--- a/Engine/Plugins/Experimental/MediaViewer/Source/MediaViewer/Private/Library/MediaViewerLibrary.cpp
+++ b/Engine/Plugins/Experimental/MediaViewer/Source/MediaViewer/Private/Library/MediaViewerLibrary.cpp
@@ -31,11 +31,11 @@ const FGuid& FMediaViewerLibrary::GetHistoryGroupId() const
 FMediaViewerLibrary::FMediaViewerLibrary(const FArgs& InArgs)
 {
 	const FText DefaultTitle = InArgs.DefaultGroupName.IsEmpty()
-		? LOCTEXT("DefaultGroup", "Default")
+		? LOCTEXT("DefaultGroup", "Saved")
 		: InArgs.DefaultGroupName;
 
 	const FText DefaultToolTip = InArgs.DefaultGroupToolTip.IsEmpty()
-		? LOCTEXT("DefaultGroupTooltip", "The default group.")
+		? LOCTEXT("DefaultGroupTooltip", "Saved items.")
 		: InArgs.DefaultGroupToolTip;
 
 	Groups.Add(MakeShared<FMediaViewerLibraryGroup>(
@@ -55,11 +55,6 @@ FMediaViewerLibrary::FMediaViewerLibrary(const FArgs& InArgs)
 
 bool FMediaViewerLibrary::CanDragDropGroup(const FGuid& InGroupId) const
 {
-	if (InGroupId == GetHistoryGroupId())
-	{
-		return false;
-	}
-
 	TSharedPtr<FMediaViewerLibraryGroup> Group = GetGroup(InGroupId);
 
 	return Group.IsValid() && !Group->IsDynamic();
@@ -69,7 +64,7 @@ bool FMediaViewerLibrary::CanDragDropItem(const FGroupItem& InItem) const
 {
 	TSharedPtr<FMediaViewerLibraryGroup> Group = GetGroup(InItem.GroupId);
 
-	return Group.IsValid() && !Group->IsDynamic() && Group->GetId() != GetHistoryGroupId();
+	return Group.IsValid() && !Group->IsDynamic();
 }
 
 const TArray<TSharedRef<const FMediaViewerLibraryGroup>>& FMediaViewerLibrary::GetGroups() const
@@ -174,6 +169,24 @@ TSharedPtr<FMediaViewerLibraryItem> FMediaViewerLibrary::GetItem(const FGuid& In
 	return nullptr;
 }
 
+TSharedPtr<FMediaViewerLibraryGroup> FMediaViewerLibrary::GetItemGroup(const FGuid& InItemId) const
+{
+	for (const TSharedRef<FMediaViewerLibraryGroup>& Group : Groups)
+	{
+		if (Group->GetId() == GetHistoryGroupId())
+		{
+			continue;
+		}
+
+		if (Group->ContainsItem(InItemId))
+		{
+			return Group;
+		}
+	}
+
+	return nullptr;
+}
+
 bool FMediaViewerLibrary::AddItem(const TSharedRef<FMediaViewerLibraryItem>& InNewItem)
 {
 	const FGuid& ItemId = InNewItem->GetId();
@@ -192,7 +205,7 @@ bool FMediaViewerLibrary::AddItemToGroup(const TSharedRef<FMediaViewerLibraryIte
 {
 	const FGuid& ItemId = InNewItem->GetId();
 
-	if (Items.Contains(ItemId))
+	if (GetItemGroup(InNewItem->GetId()).IsValid())
 	{
 		return false;
 	}
@@ -227,7 +240,7 @@ bool FMediaViewerLibrary::AddItemBelowItem(const TSharedRef<FMediaViewerLibraryI
 		return false;
 	}
 
-	if (Items.Contains(ItemId))
+	if (GetItemGroup(InNewItem->GetId()).IsValid())
 	{
 		return false;
 	}
diff --git a/Engine/Plugins/Experimental/MediaViewer/Source/MediaViewer/Private/Library/MediaViewerLibrary.h b/Engine/Plugins/Experimental/MediaViewer/Source/MediaViewer/Private/Library/MediaViewerLibrary.h
--- a/Engine/Plugins/Experimental/MediaViewer/Source/MediaViewer/Private/Library/MediaViewerLibrary.h
+++ b/Engine/Plugins/Experimental/MediaViewer/Source/MediaViewer/Private/Library/MediaViewerLibrary.h
@@ -41,6 +41,7 @@ public:
 	virtual TSharedPtr<FMediaViewerLibraryGroup> RemoveGroup(const FGuid& InGroupIdToRemove) override;
 	virtual TSharedPtr<FMediaViewerLibraryItem> FindItemByValue(FName InItemType, const FString& InStringValue) const override;
 	virtual TSharedPtr<FMediaViewerLibraryItem> GetItem(const FGuid& InItemId) const override;
+	virtual TSharedPtr<FMediaViewerLibraryGroup> GetItemGroup(const FGuid& InItemId) const override;
 	virtual bool AddItem(const TSharedRef<FMediaViewerLibraryItem>& InNewItem) override;
 	virtual bool AddItemToGroup(const TSharedRef<FMediaViewerLibraryItem>& InNewItem, TOptional<FGuid> InTargetGroupId = {},
 		int32 InIndex = INDEX_NONE) override;
diff --git a/Engine/Plugins/Experimental/MediaViewer/Source/MediaViewer/Private/Widgets/SMediaViewerLibraryGroup.cpp b/Engine/Plugins/Experimental/MediaViewer/Source/MediaViewer/Private/Widgets/SMediaViewerLibraryGroup.cpp
--- a/Engine/Plugins/Experimental/MediaViewer/Source/MediaViewer/Private/Widgets/SMediaViewerLibraryGroup.cpp
+++ b/Engine/Plugins/Experimental/MediaViewer/Source/MediaViewer/Private/Widgets/SMediaViewerLibraryGroup.cpp
@@ -229,6 +229,11 @@ bool SMediaViewerLibraryGroup::CanAcceptLibraryItem(const IMediaViewerLibrary::F
 		return false;
 	}
 
+	if (GroupId == Library->GetHistoryGroupId())
+	{
+		return false;
+	}
+
 	if (!Library->CanDragDropGroup(GroupId))
 	{
 		return false;
@@ -254,11 +259,23 @@ void SMediaViewerLibraryGroup::OnLibraryItemDropped(const IMediaViewerLibrary::F
 
 	if (GroupId == InDroppedGroupItem.GroupId)
 	{
-		LibraryWeak.Pin()->MoveItemWithinGroup(InDroppedGroupItem, Library->GetGroup(GroupId)->GetItems().Num());
+		Library->MoveItemWithinGroup(InDroppedGroupItem, Library->GetGroup(GroupId)->GetItems().Num());
+	}
+	else if (InDroppedGroupItem.GroupId != Library->GetHistoryGroupId())
+	{
+		Library->MoveItemToGroup(InDroppedGroupItem, GroupId);
 	}
 	else
 	{
-		LibraryWeak.Pin()->MoveItemToGroup(InDroppedGroupItem, GroupId);
+		if (TSharedPtr<FMediaViewerLibraryItem> Item = Library->GetItem(InDroppedGroupItem.ItemId))
+		{
+			if (TSharedPtr<FMediaViewerLibraryGroup> CurrentGroup = Library->GetItemGroup(InDroppedGroupItem.ItemId))
+			{
+				CurrentGroup->RemoveItem(InDroppedGroupItem.ItemId);
+			}
+
+			Library->AddItemToGroup(Item.ToSharedRef(), GroupId);
+		}
 	}
 }
 
diff --git a/Engine/Plugins/Experimental/MediaViewer/Source/MediaViewer/Private/Widgets/SMediaViewerLibraryItem.cpp b/Engine/Plugins/Experimental/MediaViewer/Source/MediaViewer/Private/Widgets/SMediaViewerLibraryItem.cpp
--- a/Engine/Plugins/Experimental/MediaViewer/Source/MediaViewer/Private/Widgets/SMediaViewerLibraryItem.cpp
+++ b/Engine/Plugins/Experimental/MediaViewer/Source/MediaViewer/Private/Widgets/SMediaViewerLibraryItem.cpp
@@ -334,6 +334,11 @@ bool SMediaViewerLibraryItem::CanAcceptLibraryItem(const IMediaViewerLibrary::FG
 		return false;
 	}
 
+	if (GroupItem.GroupId == Library->GetHistoryGroupId())
+	{
+		return false;
+	}
+
 	if (!Library->CanDragDropGroup(GroupItem.GroupId))
 	{
 		return false;
@@ -358,7 +363,7 @@ bool SMediaViewerLibraryItem::CanAcceptLibraryItem(const IMediaViewerLibrary::FG
 		return false;
 	}
 
-	return MyGroup->ContainsItem(GroupItem.ItemId);
+	return MyGroup->GetId() != Library->GetHistoryGroupId();
 }
 
 void SMediaViewerLibraryItem::OnLibraryItemDropped(const IMediaViewerLibrary::FGroupItem& InDroppedGroupItem) const
@@ -373,6 +378,7 @@ void SMediaViewerLibraryItem::OnLibraryItemDropped(const IMediaViewerLibrary::FG
 	TSharedPtr<FMediaViewerLibraryGroup> MyGroup = Library->GetGroup(GroupItem.GroupId);
 	const int32 MyGroupIndex = MyGroup->FindItemIndex(GroupItem.ItemId);
 
+
 	if (InDroppedGroupItem.GroupId == GroupItem.GroupId)
 	{
 		Library->MoveItemWithinGroup(
@@ -382,7 +388,7 @@ void SMediaViewerLibraryItem::OnLibraryItemDropped(const IMediaViewerLibrary::FG
 				: MyGroupIndex + 1
 		);
 	}
-	else
+	else if (InDroppedGroupItem.GroupId != Library->GetHistoryGroupId())
 	{
 		Library->MoveItemToGroup(
 			InDroppedGroupItem,
@@ -392,6 +398,21 @@ void SMediaViewerLibraryItem::OnLibraryItemDropped(const IMediaViewerLibrary::FG
 				: MyGroupIndex + 1
 		);
 	}
+	else
+	{
+		if (TSharedPtr<FMediaViewerLibraryGroup> CurrentGroup = Library->GetItemGroup(InDroppedGroupItem.ItemId))
+		{
+			CurrentGroup->RemoveItem(InDroppedGroupItem.ItemId);
+		}
+
+		Library->AddItemToGroup(
+			Library->GetItem(InDroppedGroupItem.ItemId).ToSharedRef(),
+			GroupItem.GroupId,
+			ItemDropZone.Get(EItemDropZone::OntoItem) == EItemDropZone::AboveItem
+				? MyGroupIndex
+				: MyGroupIndex + 1
+		);
+	}
 }
 
 void SMediaViewerLibraryItem::OnDragEnter(const FGeometry& InMyGeometry, const FDragDropEvent& InEvent)
diff --git a/Engine/Plugins/Experimental/MediaViewer/Source/MediaViewer/Public/Library/IMediaViewerLibrary.h b/Engine/Plugins/Experimental/MediaViewer/Source/MediaViewer/Public/Library/IMediaViewerLibrary.h
--- a/Engine/Plugins/Experimental/MediaViewer/Source/MediaViewer/Public/Library/IMediaViewerLibrary.h
+++ b/Engine/Plugins/Experimental/MediaViewer/Source/MediaViewer/Public/Library/IMediaViewerLibrary.h
@@ -115,6 +115,13 @@ public:
 	 */
 	virtual TSharedPtr<FMediaViewerLibraryItem> GetItem(const FGuid& InItemId) const = 0;
 
+	/**
+	 * Finds the first group containing the given Item Id.
+	 * @param InItemId The unique identifier of the item.
+	 * @return The group containing the item or null.
+	 */
+	virtual TSharedPtr<FMediaViewerLibraryGroup> GetItemGroup(const FGuid& InItemId) const = 0;
+
 	/**
 	 * Adds an item to the Library without place it in a group
 	 * @param InNewItem Data for the item that is being added.
