Media Viewer: Fixed viewport rendering
- Added slate effect draw flags
- Fixed pixel sampling-based crash.

diff --git a/Engine/Plugins/Experimental/MediaViewer/Source/MediaViewer/Private/ImageViewers/MaterialInterfaceImageViewer.cpp b/Engine/Plugins/Experimental/MediaViewer/Source/MediaViewer/Private/ImageViewers/MaterialInterfaceImageViewer.cpp
--- a/Engine/Plugins/Experimental/MediaViewer/Source/MediaViewer/Private/ImageViewers/MaterialInterfaceImageViewer.cpp
+++ b/Engine/Plugins/Experimental/MediaViewer/Source/MediaViewer/Private/ImageViewers/MaterialInterfaceImageViewer.cpp
@@ -214,6 +214,8 @@ FMaterialInterfaceImageViewer::FMaterialInterfaceImageViewer(const FGuid& InId,
 		FText::FromName(InMaterialInterface->GetFName())
 	})
 {
+	DrawEffects |= ESlateDrawEffect::PreMultipliedAlpha | ESlateDrawEffect::NoGamma;
+
 	MaterialSettings.MaterialInterface = InMaterialInterface;
 
 	CreateBrush();
@@ -249,6 +251,11 @@ TSharedPtr<FMediaViewerLibraryItem> FMaterialInterfaceImageViewer::CreateLibrary
 
 TOptional<TVariant<FColor, FLinearColor>> FMaterialInterfaceImageViewer::GetPixelColor(const FIntPoint& InPixelCoords, int32 InMipLevel) const
 {
+	if (!SampleCache.IsValid())
+	{
+		return {};
+	}
+
 	if (InPixelCoords.X < 0 || InPixelCoords.Y < 0)
 	{
 		SampleCache.Invalidate();
@@ -322,9 +329,9 @@ void FMaterialInterfaceImageViewer::CreateBrush()
 	UTextureRenderTarget2D* RenderTarget = FMediaViewerUtils::RenderMaterial(MaterialSettings.MaterialInterface);
 	MaterialSettings.RenderTarget = RenderTarget;
 
-	SampleCache = FTextureSampleCache(RenderTarget);
+	Brush = MakeShared<FSlateImageBrush>(RenderTarget, FVector2D(ImageInfo.Size.X, ImageInfo.Size.Y));
 
-	Brush = MakeShared<FSlateImageBrush>(RenderTarget, FVector2D(RenderTarget->GetSurfaceWidth(), RenderTarget->GetSurfaceDepth()));
+	SampleCache = FTextureSampleCache(RenderTarget, RenderTarget->GetFormat());
 }
 
 void FMaterialInterfaceImageViewer::RenderMaterial()
@@ -332,7 +339,11 @@ void FMaterialInterfaceImageViewer::RenderMaterial()
 	if (::IsValid(MaterialSettings.MaterialInterface) && ::IsValid(MaterialSettings.RenderTarget))
 	{
 		FMediaViewerUtils::RenderMaterial(MaterialSettings.MaterialInterface, MaterialSettings.RenderTarget);
-		SampleCache.MarkDirty();
+
+		if (SampleCache.IsValid())
+		{
+			SampleCache.MarkDirty();
+		}
 	}
 }
 
diff --git a/Engine/Plugins/Experimental/MediaViewer/Source/MediaViewer/Private/ImageViewers/MediaImageViewer.cpp b/Engine/Plugins/Experimental/MediaViewer/Source/MediaViewer/Private/ImageViewers/MediaImageViewer.cpp
--- a/Engine/Plugins/Experimental/MediaViewer/Source/MediaViewer/Private/ImageViewers/MediaImageViewer.cpp
+++ b/Engine/Plugins/Experimental/MediaViewer/Source/MediaViewer/Private/ImageViewers/MediaImageViewer.cpp
@@ -167,13 +167,6 @@ void FMediaImageViewer::PaintImage(FMediaImagePaintParams& InPaintParams, const
 		return;
 	}
 
-	ESlateDrawEffect DrawEffects = ESlateDrawEffect::NoPixelSnapping;
-	
-	if (UTextureRenderTarget2D* RenderTarget = Cast<UTextureRenderTarget2D>(Brush->GetResourceObject()))
-	{
-		DrawEffects |= ESlateDrawEffect::PreMultipliedAlpha | ESlateDrawEffect::NoGamma;
-	}
-
 	FSlateDrawElement::MakeBox(
 		InPaintParams.DrawElements,
 		InPaintParams.LayerId,
diff --git a/Engine/Plugins/Experimental/MediaViewer/Source/MediaViewer/Private/ImageViewers/MediaSourceImageViewer.cpp b/Engine/Plugins/Experimental/MediaViewer/Source/MediaViewer/Private/ImageViewers/MediaSourceImageViewer.cpp
--- a/Engine/Plugins/Experimental/MediaViewer/Source/MediaViewer/Private/ImageViewers/MediaSourceImageViewer.cpp
+++ b/Engine/Plugins/Experimental/MediaViewer/Source/MediaViewer/Private/ImageViewers/MediaSourceImageViewer.cpp
@@ -221,7 +221,8 @@ FMediaSourceImageViewer::FMediaSourceImageViewer(const FGuid& InId, TNotNull<UMe
 				Player->Play();
 				Brush = MakeShared<FSlateImageBrush>(MediaTexture, FVector2D(ImageInfo.Size.X, ImageInfo.Size.Y));
 
-				SampleCache = FTextureSampleCache(MediaTexture);
+				// We can't know the pixel format of the media texture at this point.
+				SampleCache = FTextureSampleCache(MediaTexture, PF_Unknown);
 			}
 		}
 	}
@@ -252,6 +253,11 @@ TSharedPtr<FMediaViewerLibraryItem> FMediaSourceImageViewer::CreateLibraryItem()
 
 TOptional<TVariant<FColor, FLinearColor>> FMediaSourceImageViewer::GetPixelColor(const FIntPoint& InPixelCoords, int32 InMipLevel) const
 {
+	if (!SampleCache.IsValid())
+	{
+		return {};
+	}
+	
 	if (InPixelCoords.X < 0 || InPixelCoords.Y < 0)
 	{
 		SampleCache.Invalidate();
diff --git a/Engine/Plugins/Experimental/MediaViewer/Source/MediaViewer/Private/ImageViewers/MediaTextureImageViewer.cpp b/Engine/Plugins/Experimental/MediaViewer/Source/MediaViewer/Private/ImageViewers/MediaTextureImageViewer.cpp
--- a/Engine/Plugins/Experimental/MediaViewer/Source/MediaViewer/Private/ImageViewers/MediaTextureImageViewer.cpp
+++ b/Engine/Plugins/Experimental/MediaViewer/Source/MediaViewer/Private/ImageViewers/MediaTextureImageViewer.cpp
@@ -182,7 +182,8 @@ FMediaTextureImageViewer::FMediaTextureImageViewer(const FGuid& InId, TNotNull<U
 		FVector2D(ImageInfo.Size.X, ImageInfo.Size.Y)
 	);
 
-	SampleCache = FTextureSampleCache(InMediaTexture);
+	// We can't know the pixel format of the media texture at this point.
+	SampleCache = FTextureSampleCache(InMediaTexture, PF_Unknown);
 }
 
 TSharedPtr<FMediaViewerLibraryItem> FMediaTextureImageViewer::CreateLibraryItem() const
@@ -205,6 +206,11 @@ TSharedPtr<FMediaViewerLibraryItem> FMediaTextureImageViewer::CreateLibraryItem(
 
 TOptional<TVariant<FColor, FLinearColor>> FMediaTextureImageViewer::GetPixelColor(const FIntPoint& InPixelCoords, int32 InMipLevel) const
 {
+	if (!SampleCache.IsValid())
+	{
+		return {};
+	}
+
 	// Always mark the sample cache dirty. We have no control over the state of the media texture.
 	SampleCache.MarkDirty();
 
diff --git a/Engine/Plugins/Experimental/MediaViewer/Source/MediaViewer/Private/ImageViewers/Texture2DImageViewer.cpp b/Engine/Plugins/Experimental/MediaViewer/Source/MediaViewer/Private/ImageViewers/Texture2DImageViewer.cpp
--- a/Engine/Plugins/Experimental/MediaViewer/Source/MediaViewer/Private/ImageViewers/Texture2DImageViewer.cpp
+++ b/Engine/Plugins/Experimental/MediaViewer/Source/MediaViewer/Private/ImageViewers/Texture2DImageViewer.cpp
@@ -197,7 +197,7 @@ FTexture2DImageViewer::FTexture2DImageViewer(const FGuid& InId, TNotNull<UTextur
 		InTexture->GetPlatformData()->Mips.Num(),
 		FText::FromName(InTexture->GetFName())
 	})
-	, SampleCache(InTexture)
+	, SampleCache(InTexture, InTexture->GetPixelFormat())
 {
 	TextureSettings.Texture = InTexture;
 
@@ -227,6 +227,11 @@ TSharedPtr<FMediaViewerLibraryItem> FTexture2DImageViewer::CreateLibraryItem() c
 
 TOptional<TVariant<FColor, FLinearColor>> FTexture2DImageViewer::GetPixelColor(const FIntPoint& InPixelCoords, int32 InMipLevel) const
 {
+	if (!SampleCache.IsValid())
+	{
+		return {};
+	}
+
 	if (InPixelCoords.X < 0 || InPixelCoords.Y < 0)
 	{
 		return {};
diff --git a/Engine/Plugins/Experimental/MediaViewer/Source/MediaViewer/Private/ImageViewers/TextureRenderTarget2DImageViewer.cpp b/Engine/Plugins/Experimental/MediaViewer/Source/MediaViewer/Private/ImageViewers/TextureRenderTarget2DImageViewer.cpp
--- a/Engine/Plugins/Experimental/MediaViewer/Source/MediaViewer/Private/ImageViewers/TextureRenderTarget2DImageViewer.cpp
+++ b/Engine/Plugins/Experimental/MediaViewer/Source/MediaViewer/Private/ImageViewers/TextureRenderTarget2DImageViewer.cpp
@@ -167,6 +167,8 @@ FTextureRenderTarget2DImageViewer::FTextureRenderTarget2DImageViewer(const FGuid
 		FText::FromName(InRenderTarget->GetFName())
 	})
 {
+	DrawEffects |= ESlateDrawEffect::PreMultipliedAlpha | ESlateDrawEffect::NoGamma;
+
 	RenderTargetSettings.RenderTarget = InRenderTarget;
 
 	Brush = MakeShared<FSlateImageBrush>(
@@ -174,7 +176,7 @@ FTextureRenderTarget2DImageViewer::FTextureRenderTarget2DImageViewer(const FGuid
 		FVector2D(InRenderTarget->GetSurfaceWidth(), InRenderTarget->GetSurfaceHeight())
 	);
 	
-	SampleCache = FTextureSampleCache(InRenderTarget);
+	SampleCache = FTextureSampleCache(InRenderTarget, InRenderTarget->GetFormat());
 }
 
 TSharedPtr<FMediaViewerLibraryItem> FTextureRenderTarget2DImageViewer::CreateLibraryItem() const
@@ -197,6 +199,11 @@ TSharedPtr<FMediaViewerLibraryItem> FTextureRenderTarget2DImageViewer::CreateLib
 
 TOptional<TVariant<FColor, FLinearColor>> FTextureRenderTarget2DImageViewer::GetPixelColor(const FIntPoint& InPixelCoords, int32 InMipLevel) const
 {
+	if (!SampleCache.IsValid())
+	{
+		return {};
+	}
+
 	// Always mark the sample cache dirty. We have no control over the state of the media texture.
 	SampleCache.MarkDirty();
 
diff --git a/Engine/Plugins/Experimental/MediaViewer/Source/MediaViewer/Private/ImageViewers/TextureSampleCache.cpp b/Engine/Plugins/Experimental/MediaViewer/Source/MediaViewer/Private/ImageViewers/TextureSampleCache.cpp
--- a/Engine/Plugins/Experimental/MediaViewer/Source/MediaViewer/Private/ImageViewers/TextureSampleCache.cpp
+++ b/Engine/Plugins/Experimental/MediaViewer/Source/MediaViewer/Private/ImageViewers/TextureSampleCache.cpp
@@ -3,19 +3,48 @@
 #include "ImageViewers/TextureSampleCache.h"
 
 #include "Engine/Texture.h"
+#include "ImageViewers/TextureMipCache.h"
+#include "PixelFormat.h"
 #include "RenderingThread.h"
 #include "TextureResource.h"
+#include "UObject/StrongObjectPtr.h"
 
 namespace UE::MediaViewer::Private
 {
 
+bool CanCopyTextureWithRHICommand(EPixelFormat InFormat)
+{
+	switch (InFormat)
+	{
+		// DirectX compressed textures cannot read using the RHI ReadSurface method.
+		case EPixelFormat::PF_DXT1:
+		case EPixelFormat::PF_DXT3:
+		case EPixelFormat::PF_DXT5:
+			return false;
+
+		default:
+			return true;
+	}
+}
+
 FTextureSampleCache::FTextureSampleCache()
+	: PixelFormat(PF_Unknown)
 {
 }
 
-FTextureSampleCache::FTextureSampleCache(TNotNull<UTexture*> InTexture)
+FTextureSampleCache::FTextureSampleCache(TNotNull<UTexture*> InTexture, EPixelFormat InPixelFormat)
 	: Texture(InTexture)
+	, PixelFormat(InPixelFormat)
+{
+	if (!CanCopyTextureWithRHICommand(PixelFormat))
+	{
+		MipCache = MakeShared<FTextureMipCache>(InTexture);
+	}
+}
+
+bool FTextureSampleCache::IsValid() const
 {
+	return ::IsValid(Texture);
 }
 
 const FLinearColor* FTextureSampleCache::GetPixelColor(const FIntPoint& InPixelCoordinates, const TOptional<FTimespan>& InTime)
@@ -30,32 +59,14 @@ const FLinearColor* FTextureSampleCache::GetPixelColor(const FIntPoint& InPixelC
 
 	if (bNeedsUpdate)
 	{
-		ENQUEUE_RENDER_COMMAND(GetPixelColors)
-		(
-			[this, InPixelCoordinates, InTime](FRHICommandListImmediate& RHICmdList)
-			{
-				if (FTextureResource* TextureResource = Texture->GetResource())
-				{
-					const FIntVector Size = TextureResource->TextureRHI->GetSizeXYZ();
-					TArray<FLinearColor> Data;
-					FIntRect Rect(InPixelCoordinates.X, InPixelCoordinates.Y, InPixelCoordinates.X + 1, InPixelCoordinates.Y + 1);
-					RHICmdList.ReadSurfaceData(TextureResource->TextureRHI, Rect, Data, FReadSurfaceDataFlags());
-
-					if (!Data.IsEmpty())
-					{
-						PixelColorSample = {
-							InPixelCoordinates,
-							InTime,
-							Data[0]
-						};
-						bDirty = false;
-						return;
-					}
-				}
-
-				PixelColorSample.Reset();
-			}
-		);
+		if (!CanCopyTextureWithRHICommand(PixelFormat))
+		{
+			SetPixelColor_Mip(InPixelCoordinates, InTime);
+		}
+		else
+		{
+			SetPixelColor_RHI(InPixelCoordinates, InTime);
+		}
 	}
 
 	if (PixelColorSample.IsSet())
@@ -74,6 +85,66 @@ void FTextureSampleCache::MarkDirty()
 void FTextureSampleCache::Invalidate()
 {
 	PixelColorSample.Reset();
+
+	if (MipCache.IsValid())
+	{
+		MipCache->Invalidate();
+	}
+}
+
+void FTextureSampleCache::SetPixelColor_RHI(const FIntPoint& InPixelCoordinates, const TOptional<FTimespan>& InTime)
+{
+	ENQUEUE_RENDER_COMMAND(GetPixelColors)
+	(
+		[this, InPixelCoordinates, InTime](FRHICommandListImmediate& RHICmdList)
+		{
+			if (FTextureResource* TextureResource = Texture->GetResource())
+			{
+				const FIntVector Size = TextureResource->TextureRHI->GetSizeXYZ();
+				TArray<FLinearColor> Data;
+				FIntRect Rect(InPixelCoordinates.X, InPixelCoordinates.Y, InPixelCoordinates.X + 1, InPixelCoordinates.Y + 1);
+				RHICmdList.ReadSurfaceData(TextureResource->TextureRHI, Rect, Data, FReadSurfaceDataFlags());
+
+				if (!Data.IsEmpty())
+				{
+					PixelColorSample = {
+						InPixelCoordinates,
+						InTime,
+						Data[0]
+					};
+					bDirty = false;
+					return;
+				}
+			}
+
+			PixelColorSample.Reset();
+		}
+	);
+}
+
+void FTextureSampleCache::SetPixelColor_Mip(const FIntPoint& InPixelCoordinates, const TOptional<FTimespan>& InTime)
+{
+	if (!MipCache.IsValid())
+	{
+		return;
+	}
+
+	PixelColorSample.Reset();
+
+	const FImage* Mip = MipCache->GetMipImage(/* Mip level */ 0);
+
+	if (!Mip)
+	{
+		return;
+	}
+
+	PixelColorSample = {
+		InPixelCoordinates,
+		InTime,
+		Mip->GetOnePixelLinear(InPixelCoordinates.X, InPixelCoordinates.Y)
+	};
+
+	bDirty = false;
 }
 
 } // UE::MediaViewer::Private
diff --git a/Engine/Plugins/Experimental/MediaViewer/Source/MediaViewer/Private/ImageViewers/TextureSampleCache.h b/Engine/Plugins/Experimental/MediaViewer/Source/MediaViewer/Private/ImageViewers/TextureSampleCache.h
--- a/Engine/Plugins/Experimental/MediaViewer/Source/MediaViewer/Private/ImageViewers/TextureSampleCache.h
+++ b/Engine/Plugins/Experimental/MediaViewer/Source/MediaViewer/Private/ImageViewers/TextureSampleCache.h
@@ -2,6 +2,8 @@
 
 #pragma once
 
+#include "Templates/SharedPointer.h"
+
 #include "Math/Color.h"
 #include "Math/IntPoint.h"
 #include "Misc/NotNull.h"
@@ -9,15 +11,20 @@
 #include "Misc/Timespan.h"
 
 class UTexture;
+enum EPixelFormat : uint8;
 
 namespace UE::MediaViewer::Private
 {
 
+class FTextureMipCache;
+
 class FTextureSampleCache
 {
 public:
 	FTextureSampleCache();
-	FTextureSampleCache(TNotNull<UTexture*> InTexture);
+	FTextureSampleCache(TNotNull<UTexture*> InTexture, EPixelFormat InPixelFormat);
+
+	bool IsValid() const;
 
 	const FLinearColor* GetPixelColor(const FIntPoint& InPixelCoordinates, const TOptional<FTimespan>& InTime = {});
 
@@ -34,8 +41,14 @@ protected:
 	};
 
 	UTexture* Texture = nullptr;
+	EPixelFormat PixelFormat;
 	bool bDirty = true;
 	TOptional<FPixelColorSample> PixelColorSample;
+	TSharedPtr<FTextureMipCache> MipCache;
+
+	void SetPixelColor_RHI(const FIntPoint& InPixelCoordinates, const TOptional<FTimespan>& InTime);
+
+	void SetPixelColor_Mip(const FIntPoint& InPixelCoordinates, const TOptional<FTimespan>& InTime);
 };
 
 } // UE::MediaViewer::Private
diff --git a/Engine/Plugins/Experimental/MediaViewer/Source/MediaViewer/Private/ImageViewers/ViewportImageViewer.cpp b/Engine/Plugins/Experimental/MediaViewer/Source/MediaViewer/Private/ImageViewers/ViewportImageViewer.cpp
--- a/Engine/Plugins/Experimental/MediaViewer/Source/MediaViewer/Private/ImageViewers/ViewportImageViewer.cpp
+++ b/Engine/Plugins/Experimental/MediaViewer/Source/MediaViewer/Private/ImageViewers/ViewportImageViewer.cpp
@@ -20,6 +20,7 @@ namespace UE::MediaViewer::Private
 FViewportImageViewer::FViewportImageViewer(const FMediaImageViewerInfo& InImageInfo)
 	: FMediaImageViewer(InImageInfo)
 {
+	DrawEffects |= ESlateDrawEffect::InvertAlpha | ESlateDrawEffect::PreMultipliedAlpha | ESlateDrawEffect::NoGamma;
 }
 
 FViewportImageViewer::~FViewportImageViewer()
@@ -33,6 +34,11 @@ bool FViewportImageViewer::IsValid() const
 
 TOptional<TVariant<FColor, FLinearColor>> FViewportImageViewer::GetPixelColor(const FIntPoint& InPixelCoords, int32 InMipLevel) const
 {
+	if (!SampleCache.IsValid())
+	{
+		return {};
+	}
+
 	if (InPixelCoords.X < 0 || InPixelCoords.Y < 0)
 	{
 		SampleCache.Invalidate();
@@ -104,7 +110,7 @@ void FViewportImageViewer::CreateBrush()
 
 	Brush = MakeShared<FSlateImageBrush>(RenderTarget, FVector2D(RenderTarget->GetSurfaceWidth(), RenderTarget->GetSurfaceDepth()));
 
-	SampleCache = FTextureSampleCache(RenderTarget);
+	SampleCache = FTextureSampleCache(RenderTarget, RenderTarget->GetFormat());
 }
 
 void FViewportImageViewer::RenderViewport(TNotNull<FViewport*> InViewport, TNotNull<UTextureRenderTarget2D*> InRenderTarget, 
@@ -169,7 +175,10 @@ void FViewportImageViewer::RenderViewport()
 
 	RenderViewport(TNotNull<FViewport*>(Viewport), TNotNull<UTextureRenderTarget2D*>(ViweportSettings.RenderTarget));
 
-	SampleCache.MarkDirty();
+	if (SampleCache.IsValid())
+	{
+		SampleCache.MarkDirty();
+	}
 }
 
 } // UE::MediaViewer::Private
diff --git a/Engine/Plugins/Experimental/MediaViewer/Source/MediaViewer/Public/ImageViewer/MediaImageViewer.h b/Engine/Plugins/Experimental/MediaViewer/Source/MediaViewer/Public/ImageViewer/MediaImageViewer.h
--- a/Engine/Plugins/Experimental/MediaViewer/Source/MediaViewer/Public/ImageViewer/MediaImageViewer.h
+++ b/Engine/Plugins/Experimental/MediaViewer/Source/MediaViewer/Public/ImageViewer/MediaImageViewer.h
@@ -15,6 +15,7 @@
 #include "MediaViewer.h"
 #include "Misc/Guid.h"
 #include "Misc/Optional.h"
+#include "Rendering/RenderingCommon.h"
 #include "UObject/SoftObjectPtr.h"
 #include "UObject/StructOnScope.h"
 
@@ -283,6 +284,7 @@ protected:
 	FMediaImagePaintSettings PaintSettings;
 	FSlateImageBrush BackgroundImageBrush;
 	TSharedPtr<FSlateBrush> Brush;
+	ESlateDrawEffect DrawEffects = ESlateDrawEffect::NoPixelSnapping;
 
 	FSlateClippingZone CreateSlateClippingZone(const FGeometry& InGeometry, const FVector2D& InViewerSize,
 		const FVector2D& InViewerPosition, EOrientation InOrientation, const FFloatRange& InUVRange) const;
