Media Viewer: Group filter for library widgets

diff --git a/Engine/Plugins/Experimental/MediaViewer/Source/MediaViewer/Private/Widgets/SMediaViewerLibrary.cpp b/Engine/Plugins/Experimental/MediaViewer/Source/MediaViewer/Private/Widgets/SMediaViewerLibrary.cpp
--- a/Engine/Plugins/Experimental/MediaViewer/Source/MediaViewer/Private/Widgets/SMediaViewerLibrary.cpp
+++ b/Engine/Plugins/Experimental/MediaViewer/Source/MediaViewer/Private/Widgets/SMediaViewerLibrary.cpp
@@ -28,6 +28,7 @@ void SMediaViewerLibrary::PrivateRegisterAttributes(FSlateAttributeDescriptor::F
 void SMediaViewerLibrary::Construct(const FArguments& InArgs, const IMediaViewerLibraryWidget::FArgs& InMediaViewerLibraryArgs,
 	const TSharedRef<FMediaViewerDelegates>& InDelegates)
 {
+	GroupFilter = InMediaViewerLibraryArgs.GroupFilter;
 	Delegates = InDelegates;
 	Library = MakeShared<FMediaViewerLibrary>(InMediaViewerLibraryArgs.LibraryArgs);
 	Library->GetOnChanged().BindSP(this, &SMediaViewerLibrary::OnLibraryChanged);
@@ -93,9 +94,14 @@ void SMediaViewerLibrary::UpdateGroups()
 
 	Groups.Empty(LibraryImpl->GetGroups().Num());
 
+	TSharedRef<const IMediaViewerLibrary> ConstLibrary = Library.ToSharedRef();
+
 	for (const TSharedRef<const FMediaViewerLibraryGroup>& Group : LibraryImpl->GetGroups())
 	{
-		Groups.Add({Group->GetId(), FGuid()});
+		if (!GroupFilter.IsBound() || GroupFilter.Execute(ConstLibrary, Group))
+		{
+			Groups.Add({Group->GetId(), FGuid()});
+		}
 	}
 }
 
diff --git a/Engine/Plugins/Experimental/MediaViewer/Source/MediaViewer/Private/Widgets/SMediaViewerLibrary.h b/Engine/Plugins/Experimental/MediaViewer/Source/MediaViewer/Private/Widgets/SMediaViewerLibrary.h
--- a/Engine/Plugins/Experimental/MediaViewer/Source/MediaViewer/Private/Widgets/SMediaViewerLibrary.h
+++ b/Engine/Plugins/Experimental/MediaViewer/Source/MediaViewer/Private/Widgets/SMediaViewerLibrary.h
@@ -57,6 +57,7 @@ public:
 protected:
 	TSharedPtr<FMediaViewerDelegates> Delegates;
 	TSharedPtr<FMediaViewerLibrary> Library;
+	FGroupFilter GroupFilter;
 	TArray<IMediaViewerLibrary::FGroupItem> Groups;
 	TSharedPtr<STreeView<IMediaViewerLibrary::FGroupItem>> TreeView;
 
diff --git a/Engine/Plugins/Experimental/MediaViewer/Source/MediaViewer/Private/Widgets/SMediaViewerLibraryPrivate.cpp b/Engine/Plugins/Experimental/MediaViewer/Source/MediaViewer/Private/Widgets/SMediaViewerLibraryPrivate.cpp
--- a/Engine/Plugins/Experimental/MediaViewer/Source/MediaViewer/Private/Widgets/SMediaViewerLibraryPrivate.cpp
+++ b/Engine/Plugins/Experimental/MediaViewer/Source/MediaViewer/Private/Widgets/SMediaViewerLibraryPrivate.cpp
@@ -21,7 +21,6 @@
 namespace UE::MediaViewer::Private
 {
 	constexpr int32 MaxHistoryEntries = 25;
-	static const FGuid ColorGroupId = FGuid(0x27FDA3CE, 0x575E24CB, 0xE3717531, 0xCA9D2169);
 }
 
 namespace UE::MediaViewer::Private
@@ -38,17 +37,13 @@ void SMediaViewerLibraryPrivate::PrivateRegisterAttributes(FSlateAttributeDescri
 void SMediaViewerLibraryPrivate::Construct(const FArguments& InArgs, const TSharedRef<FMediaViewerDelegates>& InDelegates)
 {
 	IMediaViewerLibraryWidget::FArgs Args;
+	Args.GroupFilter = InArgs._GroupFilter;
 
 	ChildSlot
 	[
 		SAssignNew(Library, SMediaViewerLibrary, Args, InDelegates)
 	];
 
-	if (!::UMediaViewerLibraryIni::Get().HasGroup(UE::MediaViewer::Private::ColorGroupId))
-	{
-		AddColors();
-	}
-
 	TSharedRef<FMediaViewerLibrary> LibraryImpl = GetLibrary();
 
 	LibraryImpl->AddGroup(MakeShared<FMediaTextureGroup>(LibraryImpl));
@@ -107,60 +102,6 @@ void SMediaViewerLibraryPrivate::OnImageViewerOpened(const TSharedRef<FMediaImag
 	Library->OnLibraryChanged(Library->GetLibrary(), IMediaViewerLibrary::EChangeType::ItemGroupChanged);
 }
 
-void SMediaViewerLibraryPrivate::AddColors()
-{
-	TSharedRef<IMediaViewerLibrary> LibraryImpl = Library->GetLibrary();
-	TSharedPtr<FMediaViewerLibraryGroup> ColorGroup = LibraryImpl->GetGroup(UE::MediaViewer::Private::ColorGroupId);
-
-	if (!ColorGroup.IsValid())
-	{
-		ColorGroup = MakeShared<FMediaViewerLibraryGroup>(
-			UE::MediaViewer::Private::ColorGroupId,
-			LOCTEXT("Colors", "Colors"),
-			LOCTEXT("ColorsToolTip", "A set of colors"),
-			/* Dynamic */ false
-		);
-
-		LibraryImpl->AddGroup(ColorGroup.ToSharedRef());
-	}
-	
-	static TSharedRef<FColorImageViewer::FItem> ColorItems[5] = {
-		MakeShared<FColorImageViewer::FItem>(FGuid::Combine(UE::MediaViewer::Private::ColorGroupId, FGuid(0, 0, 0, 1)),
-			INVTEXT("Red"), INVTEXT("It's Red!"), FLinearColor::Red),
-
-		MakeShared<FColorImageViewer::FItem>(FGuid::Combine(UE::MediaViewer::Private::ColorGroupId, FGuid(0, 0, 0, 2)),
-			INVTEXT("Green"), INVTEXT("It's Green!"), FLinearColor::Green),
-
-		MakeShared<FColorImageViewer::FItem>(FGuid::Combine(UE::MediaViewer::Private::ColorGroupId, FGuid(0, 0, 0, 3)),
-			INVTEXT("Blue"), INVTEXT("It's Blue!"), FLinearColor::Blue),
-
-		MakeShared<FColorImageViewer::FItem>(FGuid::Combine(UE::MediaViewer::Private::ColorGroupId, FGuid(0, 0, 0, 4)),
-			INVTEXT("White"), INVTEXT("It's White!"), FLinearColor::White),
-
-		MakeShared<FColorImageViewer::FItem>(FGuid::Combine(UE::MediaViewer::Private::ColorGroupId, FGuid(0, 0, 0, 5)),
-			INVTEXT("Black"), INVTEXT("It's Black!"), FLinearColor::Black)
-	};
-
-	while (!ColorGroup->GetItems().IsEmpty())
-	{
-		ColorGroup->RemoveItem(ColorGroup->GetItems()[0]);
-	}
-
-	ColorGroup->AddItem(ColorItems[0]->GetId());
-	ColorGroup->AddItem(ColorItems[1]->GetId());
-	ColorGroup->AddItem(ColorItems[2]->GetId());
-	ColorGroup->AddItem(ColorItems[3]->GetId());
-	ColorGroup->AddItem(ColorItems[4]->GetId());
-
-	for (const TSharedRef<FColorImageViewer::FItem>& ColorItem : ColorItems)
-	{
-		if (!LibraryImpl->GetItem(ColorItem->GetId()))
-		{
-			LibraryImpl->AddItem(ColorItem);
-		}
-	}
-}
-
 } // UE::MediaViewer::Private
 
 #undef LOCTEXT_NAMESPACE
diff --git a/Engine/Plugins/Experimental/MediaViewer/Source/MediaViewer/Private/Widgets/SMediaViewerLibraryPrivate.h b/Engine/Plugins/Experimental/MediaViewer/Source/MediaViewer/Private/Widgets/SMediaViewerLibraryPrivate.h
--- a/Engine/Plugins/Experimental/MediaViewer/Source/MediaViewer/Private/Widgets/SMediaViewerLibraryPrivate.h
+++ b/Engine/Plugins/Experimental/MediaViewer/Source/MediaViewer/Private/Widgets/SMediaViewerLibraryPrivate.h
@@ -4,6 +4,8 @@
 
 #include "Widgets/SCompoundWidget.h"
 
+#include "Widgets/IMediaViewerLibraryWidget.h"
+
 struct FMediaViewerLibraryGroup;
 
 namespace UE::MediaViewer
@@ -25,6 +27,7 @@ class SMediaViewerLibraryPrivate : public SCompoundWidget
 
 	SLATE_BEGIN_ARGS(SMediaViewerLibraryPrivate)
 		{}
+		SLATE_EVENT(IMediaViewerLibraryWidget::FGroupFilter, GroupFilter)
 	SLATE_END_ARGS()
 
 public:
@@ -39,8 +42,6 @@ public:
 
 protected:
 	TSharedPtr<SMediaViewerLibrary> Library;
-
-	void AddColors();
 };
 
 } // UE::MediaViewer::Private
diff --git a/Engine/Plugins/Experimental/MediaViewer/Source/MediaViewer/Public/Widgets/IMediaViewerLibraryWidget.h b/Engine/Plugins/Experimental/MediaViewer/Source/MediaViewer/Public/Widgets/IMediaViewerLibraryWidget.h
--- a/Engine/Plugins/Experimental/MediaViewer/Source/MediaViewer/Public/Widgets/IMediaViewerLibraryWidget.h
+++ b/Engine/Plugins/Experimental/MediaViewer/Source/MediaViewer/Public/Widgets/IMediaViewerLibraryWidget.h
@@ -37,12 +37,19 @@ public:
 	 */
 	DECLARE_DELEGATE_RetVal_OneParam(TSharedPtr<SWidget>, FOnGetItemContextMenu, const TArray<FMediaViewerLibraryItem>&)
 
+	/**
+	 * Returns true to display the group. Return false to hide the group.
+	 */
+	DECLARE_DELEGATE_RetVal_TwoParams(bool, FGroupFilter, const TSharedRef<const IMediaViewerLibrary>& InLibrary,
+		const TSharedRef<const FMediaViewerLibraryGroup>& InGroup)
+
 	struct FArgs
 	{
 		IMediaViewerLibrary::FArgs LibraryArgs;
 		FOnImageViewerOpened OnImageViewerOpened;
 		FOnGetGroupContextMenu OnGetGroupContextMenu;
 		FOnGetItemContextMenu OnGetItemContextMenu;
+		FGroupFilter GroupFilter;
 	};
 
 	/**
