Media Viewer: DPI woes!

diff --git a/Engine/Plugins/Experimental/MediaViewer/Source/MediaViewer/Private/ImageViewers/MediaImageViewer.cpp b/Engine/Plugins/Experimental/MediaViewer/Source/MediaViewer/Private/ImageViewers/MediaImageViewer.cpp
--- a/Engine/Plugins/Experimental/MediaViewer/Source/MediaViewer/Private/ImageViewers/MediaImageViewer.cpp
+++ b/Engine/Plugins/Experimental/MediaViewer/Source/MediaViewer/Private/ImageViewers/MediaImageViewer.cpp
@@ -4,6 +4,7 @@
 
 #include "Engine/Texture.h"
 #include "Engine/TextureRenderTarget2D.h"
+#include "Math/UnrealMathUtility.h"
 #include "Math/Vector2D.h"
 #include "Misc/TVariant.h"
 #include "Widgets/SMediaViewer.h"
@@ -59,44 +60,46 @@ FVector2D FMediaImageViewer::GetPaintSize() const
 	return FVector2D(ImageInfo.Size.X, ImageInfo.Size.Y) * GetPaintSettings().Scale;
 }
 
-FSlateClippingZone FMediaImageViewer::CreateSlateClippingZone(const FGeometry& InGeometry, const FVector2D& InViewerSize, 
+FSlateClippingZone FMediaImageViewer::CreateSlateClippingZone(const FSlateRect& InCullingRect, float InDPIScale, 
 	const FVector2D& InViewerPosition, EOrientation InOrientation, const FFloatRange& InUVRange) const
 {
-	FSlateClippingZone Zone(InGeometry);
+	constexpr float WindowBorderPadding = 4.f;
 
-	if (InOrientation == Orient_Horizontal)
-	{
-		const float Offset = InUVRange.GetLowerBound().GetValue() * InViewerSize.X;
-		const float Width = InUVRange.Size<float>() * InViewerSize.X;
+	const float Left = InCullingRect.Left + InViewerPosition.X * InDPIScale + WindowBorderPadding;
+	const float Right = InCullingRect.Right - WindowBorderPadding;
+	const float Top = InCullingRect.Top + InViewerPosition.Y * InDPIScale + WindowBorderPadding;
+	const float Bottom = InCullingRect.Bottom - WindowBorderPadding;
 
-		Zone.TopLeft.X /= InGeometry.Scale;
-		Zone.TopLeft.X += Offset + InViewerPosition.X;
-		Zone.TopRight.X = Zone.TopLeft.X + Width;
+	FSlateClippingZone ClippingZone;
 
-		Zone.BottomLeft.X /= InGeometry.Scale;
-		Zone.BottomLeft.X += Offset + InViewerPosition.X;
-		Zone.BottomRight.X = Zone.BottomLeft.X + Width;
-	}
-	else
+	switch (InOrientation)
 	{
-		const float Offset = InUVRange.GetLowerBound().GetValue() * InViewerSize.Y;
-		const float Height = InUVRange.Size<float>() * InViewerSize.Y;
-
-		Zone.TopLeft.Y /= InGeometry.Scale;
-		Zone.TopLeft.Y += Offset + InViewerPosition.Y;
-		Zone.BottomLeft.Y = Zone.TopLeft.Y + Height;
-
-		Zone.TopRight.Y /= InGeometry.Scale;
-		Zone.TopRight.Y += Offset + InViewerPosition.Y;
-		Zone.BottomRight.Y = Zone.TopRight.Y + Height;
+		case Orient_Horizontal:
+			ClippingZone.TopLeft.Y = Top;
+			ClippingZone.TopRight.Y = Top;
+			ClippingZone.BottomLeft.Y = Bottom;
+			ClippingZone.BottomRight.Y = Bottom;
+
+			ClippingZone.TopLeft.X = FMath::Lerp(Left, Right, InUVRange.GetLowerBound().GetValue());
+			ClippingZone.TopRight.X = FMath::Lerp(Left, Right, InUVRange.GetUpperBound().GetValue());
+			ClippingZone.BottomLeft.X = ClippingZone.TopLeft.X;
+			ClippingZone.BottomRight.X = ClippingZone.TopRight.X;
+			break;
+
+		case Orient_Vertical:
+			ClippingZone.TopLeft.Y = FMath::Lerp(Top, Bottom, InUVRange.GetLowerBound().GetValue());
+			ClippingZone.BottomLeft.Y = FMath::Lerp(Top, Bottom, InUVRange.GetUpperBound().GetValue());
+			ClippingZone.TopRight.Y = ClippingZone.TopLeft.Y;
+			ClippingZone.BottomRight.Y = ClippingZone.BottomLeft.Y;
+
+			ClippingZone.TopLeft.X = Left;
+			ClippingZone.BottomLeft.X = Left;
+			ClippingZone.TopRight.X = Right;
+			ClippingZone.BottomRight.X = Right;
+			break;
 	}
 
-	Zone.TopLeft *= InGeometry.Scale;
-	Zone.TopRight *= InGeometry.Scale;
-	Zone.BottomLeft *= InGeometry.Scale;
-	Zone.BottomRight *= InGeometry.Scale;
-
-	return Zone;
+	return ClippingZone;
 }
 
 void FMediaImageViewer::Paint(FMediaImagePaintParams& InPaintParams)
@@ -112,8 +115,8 @@ void FMediaImageViewer::Paint(FMediaImagePaintParams& InPaintParams)
 	};
 
 	const FSlateClippingZone ClippingZone = CreateSlateClippingZone(
-		InPaintParams.AllottedGeometry,
-		InPaintParams.ViewerSize,
+		InPaintParams.MyCullingRect,
+		InPaintParams.DPIScale,
 		InPaintParams.ViewerPosition,
 		InPaintParams.Orientation,
 		InPaintParams.UVRange
diff --git a/Engine/Plugins/Experimental/MediaViewer/Source/MediaViewer/Private/Widgets/SMediaViewer.cpp b/Engine/Plugins/Experimental/MediaViewer/Source/MediaViewer/Private/Widgets/SMediaViewer.cpp
--- a/Engine/Plugins/Experimental/MediaViewer/Source/MediaViewer/Private/Widgets/SMediaViewer.cpp
+++ b/Engine/Plugins/Experimental/MediaViewer/Source/MediaViewer/Private/Widgets/SMediaViewer.cpp
@@ -9,6 +9,7 @@
 #include "ImageViewer/MediaImageViewer.h"
 #include "ImageViewers/NullImageViewer.h"
 #include "IMediaViewerModule.h"
+#include "Library/MediaViewerLibraryGroup.h"
 #include "MediaViewerCommands.h"
 #include "Sidebar/SSidebar.h"
 #include "Sidebar/SSidebarContainer.h"
@@ -301,8 +302,6 @@ int32 SMediaViewer::OnPaint(const FPaintArgs& InArgs, const FGeometry& InAllotte
 	const FSlateLayoutTransform ViewerTransform = FSlateLayoutTransform(ViewerPosition);
 	const FPaintGeometry ViewerPaintGeometry = InAllottedGeometry.ToPaintGeometry(ViewerSize, ViewerTransform);
 
-	InOutDrawElements.PushClip(FSlateClippingZone(ViewerPaintGeometry));
-
 	FSlateDrawElement::MakeBox(
 		InOutDrawElements,
 		InLayerId,
@@ -441,9 +440,7 @@ int32 SMediaViewer::OnPaint(const FPaintArgs& InArgs, const FGeometry& InAllotte
 			break;
 	}
 
-	InOutDrawElements.PopClip();
-
-	int32 NewLayerId = SCompoundWidget::OnPaint(InArgs, InAllottedGeometry, InMyCullingRect, InOutDrawElements, PaintParams.LayerId, InWidgetStyle, bInParentEnabled);
+	const int32 NewLayerId = SCompoundWidget::OnPaint(InArgs, InAllottedGeometry, InMyCullingRect, InOutDrawElements, PaintParams.LayerId, InWidgetStyle, bInParentEnabled);
 
 	// Doing this after rendering or it creates flicker.
 	if (ContentSlot)
diff --git a/Engine/Plugins/Experimental/MediaViewer/Source/MediaViewer/Public/ImageViewer/MediaImageViewer.h b/Engine/Plugins/Experimental/MediaViewer/Source/MediaViewer/Public/ImageViewer/MediaImageViewer.h
--- a/Engine/Plugins/Experimental/MediaViewer/Source/MediaViewer/Public/ImageViewer/MediaImageViewer.h
+++ b/Engine/Plugins/Experimental/MediaViewer/Source/MediaViewer/Public/ImageViewer/MediaImageViewer.h
@@ -286,8 +286,8 @@ protected:
 	TSharedPtr<FSlateBrush> Brush;
 	ESlateDrawEffect DrawEffects = ESlateDrawEffect::NoPixelSnapping;
 
-	FSlateClippingZone CreateSlateClippingZone(const FGeometry& InGeometry, const FVector2D& InViewerSize,
-		const FVector2D& InViewerPosition, EOrientation InOrientation, const FFloatRange& InUVRange) const;
+	FSlateClippingZone CreateSlateClippingZone(const FSlateRect& InCullingRect, float InDPIScale, const FVector2D& InViewerPosition, 
+		EOrientation InOrientation, const FFloatRange& InUVRange) const;
 
 	MEDIAVIEWER_API virtual void PaintPanel(FMediaImagePaintParams& InPaintParams, const FMediaImagePaintGeometry& InPaintGeometry);
 
