Media Viewer: Playhead background

diff --git a/Engine/Plugins/Experimental/MediaViewer/Source/MediaViewer/Private/Widgets/SMediaSourceOverlay.cpp b/Engine/Plugins/Experimental/MediaViewer/Source/MediaViewer/Private/Widgets/SMediaSourceOverlay.cpp
--- a/Engine/Plugins/Experimental/MediaViewer/Source/MediaViewer/Private/Widgets/SMediaSourceOverlay.cpp
+++ b/Engine/Plugins/Experimental/MediaViewer/Source/MediaViewer/Private/Widgets/SMediaSourceOverlay.cpp
@@ -14,6 +14,7 @@
 #include "Modules/ModuleManager.h"
 #include "Styling/SlateIconFinder.h"
 #include "Styling/StyleColors.h"
+#include "Widgets/Colors/SColorBlock.h"
 #include "Widgets/Images/SImage.h"
 #include "Widgets/Input/SButton.h"
 #include "Widgets/SBoxPanel.h"
@@ -91,7 +92,7 @@ void SMediaSourceOverlay::TrySetFrameRate()
 {
 	if (UMediaPlayer* MediaPlayer = GetMediaPlayer())
 	{
-		// IDEX_NONE here access the currently playing TrackIndex and FormatIndex.
+		// INDEX_NONE here access the currently playing TrackIndex and FormatIndex.
 		const float FrameRate = MediaPlayer->GetVideoTrackFrameRate(INDEX_NONE, INDEX_NONE);
 
 		if (!FMath::IsNearlyZero(FrameRate))
@@ -125,7 +126,18 @@ TSharedRef<SWidget> SMediaSourceOverlay::CreateSlider()
 	MediaPlayerSlider->SetSliderHandleColor(FSlateColor(EStyleColor::AccentBlue));
 	MediaPlayerSlider->SetVisibleWhenInactive(EVisibility::Visible);
 
-	return MediaPlayerSlider;
+	return SNew(SOverlay)
+		+ SOverlay::Slot()
+		.Padding(5.f, 11.f, 5.f, 11.f)
+		[
+			SNew(SColorBlock)
+			.Color(FStyleColors::Header.GetSpecifiedColor())
+			.Size(FVector2D(10.f, 2.f))
+		]
+		+ SOverlay::Slot()
+		[
+			MediaPlayerSlider
+		];
 }
 
 TSharedRef<SWidget> SMediaSourceOverlay::CreateControls()
@@ -522,7 +534,7 @@ FReply SMediaSourceOverlay::Forward_OnClicked()
 {
 	if (UMediaPlayer* MediaPlayer = GetMediaPlayer())
 	{
-		// IDEX_NONE here access the currently playing TrackIndex and FormatIndex.
+		// INDEX_NONE here access the currently playing TrackIndex and FormatIndex.
 		const double FrameRate = MediaPlayer->GetVideoTrackFrameRate(INDEX_NONE, INDEX_NONE);
 
 		if (!FMath::IsNearlyZero(FrameRate))
