Media Viewer: Scale to fit
- Automatic scale to fit for images.
- Reset transform scales to fit.
- Media source's own media textures do not get added to the media texture category.
- Media files now play on open.

diff --git a/Engine/Plugins/Experimental/MediaStream/Source/MediaStream/Private/Players/MediaStreamLocalPlayer.cpp b/Engine/Plugins/Experimental/MediaStream/Source/MediaStream/Private/Players/MediaStreamLocalPlayer.cpp
--- a/Engine/Plugins/Experimental/MediaStream/Source/MediaStream/Private/Players/MediaStreamLocalPlayer.cpp
+++ b/Engine/Plugins/Experimental/MediaStream/Source/MediaStream/Private/Players/MediaStreamLocalPlayer.cpp
@@ -227,9 +227,9 @@ void UMediaStreamLocalPlayer::ApplyPlayerConfig()
 
 	// Stop media playing on open in the editor (only in PIE)
 #if WITH_EDITOR
-	const bool bOverridePlayOnOpen = !bIsPIE && !FMediaStreamModule::Get().CanAutoplay();
+	const bool bOverridePlayOnOpen = HasAnyFlags(RF_WasLoaded) && !bIsPIE && !FMediaStreamModule::Get().CanAutoplay();
 #else
-	const bool bOverridePlayOnOpen = !FMediaStreamModule::Get().CanAutoplay();
+	const bool bOverridePlayOnOpen = HasAnyFlags(RF_WasLoaded) && !FMediaStreamModule::Get().CanAutoplay();
 #endif
 
 	if (!bOverridePlayOnOpen)
diff --git a/Engine/Plugins/Experimental/MediaViewer/Source/MediaViewer/Private/ImageViewers/MediaSourceImageViewer.cpp b/Engine/Plugins/Experimental/MediaViewer/Source/MediaViewer/Private/ImageViewers/MediaSourceImageViewer.cpp
--- a/Engine/Plugins/Experimental/MediaViewer/Source/MediaViewer/Private/ImageViewers/MediaSourceImageViewer.cpp
+++ b/Engine/Plugins/Experimental/MediaViewer/Source/MediaViewer/Private/ImageViewers/MediaSourceImageViewer.cpp
@@ -218,9 +218,11 @@ FMediaSourceImageViewer::FMediaSourceImageViewer(const FGuid& InId, TNotNull<UMe
 				Config.bLooping = true;
 				Player->SetPlayerConfig(Config);
 
-				Player->Play();
 				Brush = MakeShared<FSlateImageBrush>(MediaTexture, FVector2D(ImageInfo.Size.X, ImageInfo.Size.Y));
 
+				// Make sure it doesn't display in the media texture list
+				MediaTexture->AddAssetUserData(NewObject<UMediaViewerUserData>(MediaTexture));
+
 				// We can't know the pixel format of the media texture at this point.
 				SampleCache = FTextureSampleCache(MediaTexture, PF_Unknown);
 			}
diff --git a/Engine/Plugins/Experimental/MediaViewer/Source/MediaViewer/Private/Library/MediaTextureGroup.cpp b/Engine/Plugins/Experimental/MediaViewer/Source/MediaViewer/Private/Library/MediaTextureGroup.cpp
--- a/Engine/Plugins/Experimental/MediaViewer/Source/MediaViewer/Private/Library/MediaTextureGroup.cpp
+++ b/Engine/Plugins/Experimental/MediaViewer/Source/MediaViewer/Private/Library/MediaTextureGroup.cpp
@@ -5,6 +5,7 @@
 #include "ImageViewers/MediaTextureImageViewer.h"
 #include "Library/MediaViewerLibraryItem.h"
 #include "MediaTexture.h"
+#include "MediaViewerUtils.h"
 #include "UObject/UObjectIterator.h"
 
 #define LOCTEXT_NAMESPACE "MediaTextureGroup"
@@ -42,6 +43,11 @@ TArray<TSharedRef<FMediaViewerLibraryItem>> FMediaTextureGroup::GetMediaTextureI
 			continue;
 		}
 
+		if (MediaTexture->HasAssetUserDataOfClass(UMediaViewerUserData::StaticClass()))
+		{
+			continue;
+		}
+
 		if (TSharedPtr<FMediaViewerLibraryItem> NewItem = Factory->CreateLibraryItem(TNotNull<UMediaTexture*>(MediaTexture)))
 		{
 			MediaTextureItems.Add(NewItem.ToSharedRef());
diff --git a/Engine/Plugins/Experimental/MediaViewer/Source/MediaViewer/Private/Widgets/SMediaImageViewerOverlay.cpp b/Engine/Plugins/Experimental/MediaViewer/Source/MediaViewer/Private/Widgets/SMediaImageViewerOverlay.cpp
--- a/Engine/Plugins/Experimental/MediaViewer/Source/MediaViewer/Private/Widgets/SMediaImageViewerOverlay.cpp
+++ b/Engine/Plugins/Experimental/MediaViewer/Source/MediaViewer/Private/Widgets/SMediaImageViewerOverlay.cpp
@@ -45,6 +45,7 @@ void SMediaImageViewerOverlay::Construct(const FArguments& InArgs, EMediaImageVi
 {
 	Position = InPosition;
 	Delegates = InDelegates;
+	bScaleToFit = InArgs._bScaleToFit;
 
 	TSharedPtr<FMediaImageViewer> ImageViewer = Delegates->GetImageViewer.Execute(Position);
 
@@ -278,6 +279,29 @@ bool SMediaImageViewerOverlay::SupportsKeyboardFocus() const
 	return true;
 }
 
+void SMediaImageViewerOverlay::Tick(const FGeometry& InAllottedGeometry, const double InCurrentTime, const float InDeltaTime)
+{
+	SCompoundWidget::Tick(InAllottedGeometry, InCurrentTime, InDeltaTime);
+
+	if (bScaleToFit)
+	{
+		const FVector2D LocalSize = Delegates->GetViewerSize.Execute();
+
+		if (FMath::IsNearlyZero(LocalSize.X) || FMath::IsNearlyZero(LocalSize.Y))
+		{
+			return;
+		}
+
+		TSharedPtr<FMediaImageViewer> ImageViewer = Delegates->GetImageViewer.Execute(Position);
+
+		if (ImageViewer.IsValid() && ImageViewer->GetInfo().Size.X > 2)
+		{
+			ScaleToFit(/* Use transform lock */ false);
+			bScaleToFit = false;
+		}
+	}
+}
+
 void SMediaImageViewerOverlay::UpdateMouse(const TOptional<FVector2D>& InMousePosition)
 {
 	if (!DragButtonName.IsNone())
@@ -372,7 +396,10 @@ void SMediaImageViewerOverlay::BindCommands()
 	CommandList->MapAction(Commands.Scale400, FExecuteAction::CreateSP(this, &SMediaImageViewerOverlay::Try_SetScale, 4.0f));
 	CommandList->MapAction(Commands.Scale800, FExecuteAction::CreateSP(this, &SMediaImageViewerOverlay::Try_SetScale, 8.0f));
 
-	CommandList->MapAction(Commands.ScaleToFit, FExecuteAction::CreateSP(this, &SMediaImageViewerOverlay::ScaleToFit));
+	CommandList->MapAction(
+		Commands.ScaleToFit, 
+		FExecuteAction::CreateSP(this, &SMediaImageViewerOverlay::ScaleToFit, /* Use transform lock */ true)
+	);
 
 	CommandList->MapAction(
 		Commands.ResetTransform, 
@@ -522,7 +549,7 @@ void SMediaImageViewerOverlay::ResetTransform()
 {
 	SetOffset(FVector::ZeroVector);
 	SetRotation(FRotator::ZeroRotator);
-	SetScale(1.f);
+	ScaleToFit(/* Use transform lock */ false);
 }
 
 FVector SMediaImageViewerOverlay::GetOffset() const
@@ -630,7 +657,7 @@ void SMediaImageViewerOverlay::MultiplyScale(float InMultiple)
 	SetScale(GetScale() * InMultiple);
 }
 
-void SMediaImageViewerOverlay::ScaleToFit()
+void SMediaImageViewerOverlay::ScaleToFit(bool bInUseTransformLock)
 {
 	TSharedPtr<FMediaImageViewer> ImageViewer = Delegates->GetImageViewer.Execute(Position);
 
@@ -640,6 +667,12 @@ void SMediaImageViewerOverlay::ScaleToFit()
 	}
 
 	const FVector2D LocalSize = Delegates->GetViewerSize.Execute();
+
+	if (FMath::IsNearlyZero(LocalSize.X) || FMath::IsNearlyZero(LocalSize.Y))
+	{
+		return;
+	}
+
 	const FIntPoint& ImageSize = ImageViewer->GetInfo().Size;
 
 	const float ScaleX = LocalSize.X / static_cast<float>(ImageSize.X);
@@ -647,11 +680,25 @@ void SMediaImageViewerOverlay::ScaleToFit()
 
 	if (ScaleX < ScaleY)
 	{
-		Try_SetScale(ScaleX);
+		if (bInUseTransformLock)
+		{
+			Try_SetScale(ScaleX);
+		}
+		else
+		{
+			SetScale(ScaleX);
+		}
 	}
 	else
 	{
-		Try_SetScale(ScaleY);
+		if (bInUseTransformLock)
+		{
+			Try_SetScale(ScaleY);
+		}
+		else
+		{
+			SetScale(ScaleY);
+		}
 	}
 }
 
diff --git a/Engine/Plugins/Experimental/MediaViewer/Source/MediaViewer/Private/Widgets/SMediaImageViewerOverlay.h b/Engine/Plugins/Experimental/MediaViewer/Source/MediaViewer/Private/Widgets/SMediaImageViewerOverlay.h
--- a/Engine/Plugins/Experimental/MediaViewer/Source/MediaViewer/Private/Widgets/SMediaImageViewerOverlay.h
+++ b/Engine/Plugins/Experimental/MediaViewer/Source/MediaViewer/Private/Widgets/SMediaImageViewerOverlay.h
@@ -29,8 +29,10 @@ class SMediaImageViewerOverlay : public SCompoundWidget
 
 	SLATE_BEGIN_ARGS(SMediaImageViewerOverlay)
 		: _bShowPanelName(true)
+		, _bScaleToFit(false)
 		{}
 		SLATE_ARGUMENT(bool, bShowPanelName)
+		SLATE_ARGUMENT(bool, bScaleToFit)
 	SLATE_END_ARGS()
 
 public:
@@ -60,6 +62,9 @@ public:
 
 	/** Zooms the view in or out and maintains the cursor in the same relative position. */
 	void MultiplyScaleAroundCursor(float InMultipler);
+	
+	/** Scales the image to fit the viewer. */
+	void ScaleToFit(bool bInUseTransformLock);
 
 	//~ Begin SWidget
 	virtual FCursorReply OnCursorQuery(const FGeometry& InMyGeometry, const FPointerEvent& InCursorEvent) const override;
@@ -68,6 +73,7 @@ public:
 	virtual FReply OnMouseButtonUp(const FGeometry& InMyGeometry, const FPointerEvent& InMouseEvent) override;
 	virtual FReply OnMouseWheel(const FGeometry& InMyGeometry, const FPointerEvent& InMouseEvent) override;
 	virtual bool SupportsKeyboardFocus() const override;
+	virtual void Tick(const FGeometry& InAllottedGeometry, const double InCurrentTime, const float InDeltaTime) override;
 	//~ End SWidget
 
 protected:
@@ -88,6 +94,8 @@ protected:
 	TSharedPtr<SVerticalBox> Container;
 	/** Toggles the visibility of the overlay. */
 	bool bOverlayEnabled;
+	/** Whether, on first valid paint, it should scale to fit. */
+	bool bScaleToFit;
 
 	void BindCommands();
 
@@ -117,7 +125,6 @@ protected:
 	void Try_SetScale(float InScale); // Either local or 'send to all' if transform is locked
 	void SetScale(float InScale);
 	void MultiplyScale(float InMultiple);
-	void ScaleToFit();
 
 	/** Interact with the image viewer's entire transform. */
 	bool CanResetTransform();
diff --git a/Engine/Plugins/Experimental/MediaViewer/Source/MediaViewer/Private/Widgets/SMediaViewer.cpp b/Engine/Plugins/Experimental/MediaViewer/Source/MediaViewer/Private/Widgets/SMediaViewer.cpp
--- a/Engine/Plugins/Experimental/MediaViewer/Source/MediaViewer/Private/Widgets/SMediaViewer.cpp
+++ b/Engine/Plugins/Experimental/MediaViewer/Source/MediaViewer/Private/Widgets/SMediaViewer.cpp
@@ -55,6 +55,8 @@ void SMediaViewer::Construct(const FArguments& InArgs, const FMediaViewerArgs& I
 	ImageViewers[static_cast<int32>(EMediaImageViewerPosition::Second)] = InImageViewerSecond;
 	ActiveView = EMediaImageViewerActivePosition::Single;
 	RequestedView = EMediaImageViewerActivePosition::Single;
+	ScaleToFit[static_cast<int32>(EMediaImageViewerPosition::First)] = false;
+	ScaleToFit[static_cast<int32>(EMediaImageViewerPosition::Second)] = false;
 	ContentSlot = nullptr;
 	CursorLocalPosition = {0, 0};
 
@@ -155,6 +157,7 @@ void SMediaViewer::SetImageViewer(EMediaImageViewerPosition InPosition, const TS
 	}
 
 	ImageViewers[Index] = InImageViewer;
+	ScaleToFit[Index] = true;
 
 	if (InImageViewer != FNullImageViewer::GetNullImageViewer())
 	{
@@ -520,8 +523,18 @@ void SMediaViewer::CreateDelegates()
 
 TSharedRef<SMediaImageViewerOverlay> SMediaViewer::CreateOverlay(EMediaImageViewerPosition InPosition, bool bInShowPanelName)
 {
-	return  SNew(SMediaImageViewerOverlay, InPosition, Delegates.ToSharedRef())
-		.bShowPanelName(bInShowPanelName);
+	bool bScaleToFitImage = false;
+	const int32 Index = static_cast<int32>(InPosition);
+
+	if (ScaleToFit[Index])
+	{
+		bScaleToFitImage = true;
+		ScaleToFit[Index] = false;
+	}
+
+	return SNew(SMediaImageViewerOverlay, InPosition, Delegates.ToSharedRef())
+		.bShowPanelName(bInShowPanelName)
+		.bScaleToFit(bScaleToFitImage);
 }
 
 TSharedRef<SWidget> SMediaViewer::CreateToolbar()
diff --git a/Engine/Plugins/Experimental/MediaViewer/Source/MediaViewer/Private/Widgets/SMediaViewer.h b/Engine/Plugins/Experimental/MediaViewer/Source/MediaViewer/Private/Widgets/SMediaViewer.h
--- a/Engine/Plugins/Experimental/MediaViewer/Source/MediaViewer/Private/Widgets/SMediaViewer.h
+++ b/Engine/Plugins/Experimental/MediaViewer/Source/MediaViewer/Private/Widgets/SMediaViewer.h
@@ -129,6 +129,8 @@ protected:
 	EMediaImageViewerActivePosition ActiveView;
 	/** The requested view configuration. */
 	EMediaImageViewerActivePosition RequestedView;
+	/** Whether to scale to fit on rebuild view. */
+	bool ScaleToFit[static_cast<int32>(EMediaImageViewerPosition::COUNT)];
 
 	FVector2D CursorLocalPosition;
 
diff --git a/Engine/Plugins/Experimental/MediaViewer/Source/MediaViewer/Private/Widgets/SMediaViewerToolbar.cpp b/Engine/Plugins/Experimental/MediaViewer/Source/MediaViewer/Private/Widgets/SMediaViewerToolbar.cpp
--- a/Engine/Plugins/Experimental/MediaViewer/Source/MediaViewer/Private/Widgets/SMediaViewerToolbar.cpp
+++ b/Engine/Plugins/Experimental/MediaViewer/Source/MediaViewer/Private/Widgets/SMediaViewerToolbar.cpp
@@ -337,9 +337,12 @@ FText SMediaViewerToolbar::GetScaleMenuLabel(EMediaImageViewerPosition InPositio
 {
 	if (TSharedPtr<FMediaImageViewer> ImageViewer = Delegates->GetImageViewer.Execute(InPosition))
 	{
+		FNumberFormattingOptions FormattingFloat;
+		FormattingFloat.SetMinimumFractionalDigits(1);
+
 		return FText::Format(
 			LOCTEXT("Scale", "{0}%"),
-			FText::AsNumber(ImageViewer->GetPaintSettings().Scale * 100.f)
+			FText::AsNumber(ImageViewer->GetPaintSettings().Scale * 100.f, &FormattingFloat)
 		);
 	}
 
