Media Viewer: Supports drag and drop from external media sources

diff --git a/Engine/Plugins/Experimental/MediaStream/Source/MediaStream/Private/MediaStreamSourceBlueprintLibrary.cpp b/Engine/Plugins/Experimental/MediaStream/Source/MediaStream/Private/MediaStreamSourceBlueprintLibrary.cpp
--- a/Engine/Plugins/Experimental/MediaStream/Source/MediaStream/Private/MediaStreamSourceBlueprintLibrary.cpp
+++ b/Engine/Plugins/Experimental/MediaStream/Source/MediaStream/Private/MediaStreamSourceBlueprintLibrary.cpp
@@ -68,12 +68,6 @@ FMediaStreamSource UMediaStreamSourceBlueprintLibrary::MakeMediaSourceFromAsset(
 		return {};
 	}
 
-	if (!InObject->IsAsset())
-	{
-		UE_LOG(LogMediaStream, Error, TEXT("Subobject given instead of Asset in UMediaStreamSourceBlueprintLibrary::MakeMediaSourceFromAsset"));
-		return {};
-	}
-
 	return MakeMediaSourceFromSchemePath(InMediaStream, FMediaStreamAssetSchemeHandler::Scheme, InObject.ToSoftObjectPath().ToString());
 }
 
diff --git a/Engine/Plugins/Experimental/MediaViewer/Source/MediaViewer/Private/Widgets/SMediaViewerDropTarget.cpp b/Engine/Plugins/Experimental/MediaViewer/Source/MediaViewer/Private/Widgets/SMediaViewerDropTarget.cpp
--- a/Engine/Plugins/Experimental/MediaViewer/Source/MediaViewer/Private/Widgets/SMediaViewerDropTarget.cpp
+++ b/Engine/Plugins/Experimental/MediaViewer/Source/MediaViewer/Private/Widgets/SMediaViewerDropTarget.cpp
@@ -4,10 +4,14 @@
 
 #include "AssetSelection.h"
 #include "DragAndDrop/DecoratedDragDropOp.h"
+#include "FileMediaSource.h"
 #include "Framework/Application/SlateApplication.h"
+#include "ImageViewers/MediaSourceImageViewer.h"
 #include "IMediaViewerModule.h"
+#include "Input/DragAndDrop.h"
 #include "Library/MediaViewerLibrary.h"
 #include "Library/MediaViewerLibraryItem.h"
+#include "MediaPlayer.h"
 #include "Styling/AppStyle.h"
 #include "Styling/StyleColors.h"
 #include "Widgets/Layout/SBox.h"
@@ -100,9 +104,15 @@ FReply SMediaViewerDropTarget::OnDropped(const FGeometry& InGeometry, const FDra
 		return FReply::Handled();
 	}
 
-	if (DragDropOp->IsOfType<FMediaViewerLibraryItemDragDropOperation>())
+	if (TSharedPtr<FMediaViewerLibraryItemDragDropOperation> MediaViewerOp = InDragDropEvent.GetOperationAs<FMediaViewerLibraryItemDragDropOperation>())
 	{
-		HandleDroppedMediaViewerOp(*StaticCastSharedPtr<FMediaViewerLibraryItemDragDropOperation>(DragDropOp));
+		HandleDroppedMediaViewerOp(*MediaViewerOp);
+		return FReply::Handled();
+	}
+
+	if (TSharedPtr<FExternalDragOperation> FileOp = InDragDropEvent.GetOperationAs<FExternalDragOperation>())
+	{
+		HandleDroppeFileOp(*FileOp);
 		return FReply::Handled();
 	}
 
@@ -165,6 +175,40 @@ void SMediaViewerDropTarget::HandleDroppedAssets(const TArrayView<FAssetData> In
 	}
 }
 
+void SMediaViewerDropTarget::HandleDroppeFileOp(const FExternalDragOperation& InFileDragDropOp)
+{
+	UMediaPlayer* TestMediaPlayer = NewObject<UMediaPlayer>(GetTransientPackage());
+	TArray<UFileMediaSource*, TInlineAllocator<2>> Sources;
+
+	for (const FString& FileName : InFileDragDropOp.GetFiles())
+	{
+		if (!TestMediaPlayer->OpenFile(FileName))
+		{
+			continue;
+		}
+
+		UFileMediaSource* FileMediaSource = NewObject<UFileMediaSource>(GetTransientPackage());
+		FileMediaSource->SetFilePath(FileName);
+
+		Sources.Add(FileMediaSource);
+
+		if (Sources.Num() == 2)
+		{
+			break;
+		}
+	}
+
+	TestMediaPlayer->Close();
+
+	for (int32 Index = 0; Index < Sources.Num(); ++Index)
+	{
+		Delegates->SetImageViewer.Execute(
+			static_cast<EMediaImageViewerPosition>(Index),
+			MakeShared<FMediaSourceImageViewer>(Sources[Index])
+		);
+	}
+}
+
 bool SMediaViewerDropTarget::OnAllowDrop(TSharedPtr<FDragDropOperation> InDragDropOperation) const
 {
 	if (!InDragDropOperation.IsValid())
@@ -189,6 +233,12 @@ bool SMediaViewerDropTarget::OnAllowDrop(TSharedPtr<FDragDropOperation> InDragDr
 		return true;
 	}
 
+	if (InDragDropOperation->IsOfType<FExternalDragOperation>())
+	{
+		TSharedPtr<FExternalDragOperation> ExternalDragDropOp = StaticCastSharedPtr<FExternalDragOperation>(InDragDropOperation);
+		return ExternalDragDropOp->HasFiles();
+	}
+
 	const TArray<FAssetData> DroppedAssets = AssetUtil::ExtractAssetDataFromDrag(InDragDropOperation);
 	const TArray<FAssetData> ValidAssets = GetAssetsWithImageViewer(DroppedAssets);
 
@@ -218,6 +268,11 @@ bool SMediaViewerDropTarget::OnIsRecognized(TSharedPtr<FDragDropOperation> InDra
 		return true;
 	}
 
+	if (InDragDropOperation->IsOfType<FExternalDragOperation>())
+	{
+		return true;
+	}
+
 	const TArray<FAssetData> DroppedAssets = AssetUtil::ExtractAssetDataFromDrag(InDragDropOperation);
 
 	return !DroppedAssets.IsEmpty();
diff --git a/Engine/Plugins/Experimental/MediaViewer/Source/MediaViewer/Private/Widgets/SMediaViewerDropTarget.h b/Engine/Plugins/Experimental/MediaViewer/Source/MediaViewer/Private/Widgets/SMediaViewerDropTarget.h
--- a/Engine/Plugins/Experimental/MediaViewer/Source/MediaViewer/Private/Widgets/SMediaViewerDropTarget.h
+++ b/Engine/Plugins/Experimental/MediaViewer/Source/MediaViewer/Private/Widgets/SMediaViewerDropTarget.h
@@ -43,6 +43,8 @@ protected:
 	void HandleDroppedMediaViewerOp(const FMediaViewerLibraryItemDragDropOperation& InMediaViewerOp);
 
 	void HandleDroppedAssets(const TArrayView<FAssetData> InDroppedAssets);
+	
+	void HandleDroppeFileOp(const FExternalDragOperation& InFileDragDropOp);
 
 	//~ Begin SDropTarget
 	virtual bool OnAllowDrop(TSharedPtr<FDragDropOperation> InDragDropOperation) const override;
