Media Viewer: Improved handling of external media files
- Media files are now given different types of library entries

diff --git a/Engine/Plugins/Experimental/MediaViewer/Source/MediaViewer/Private/ImageViewers/MediaSourceImageViewer.cpp b/Engine/Plugins/Experimental/MediaViewer/Source/MediaViewer/Private/ImageViewers/MediaSourceImageViewer.cpp
--- a/Engine/Plugins/Experimental/MediaViewer/Source/MediaViewer/Private/ImageViewers/MediaSourceImageViewer.cpp
+++ b/Engine/Plugins/Experimental/MediaViewer/Source/MediaViewer/Private/ImageViewers/MediaSourceImageViewer.cpp
@@ -5,6 +5,7 @@
 #include "AssetRegistry/AssetData.h"
 #include "Brushes/SlateImageBrush.h"
 #include "Engine/Texture2D.h"
+#include "FileMediaSource.h"
 #include "IMediaStreamPlayer.h"
 #include "Materials/MaterialInstanceDynamic.h"
 #include "Materials/MaterialInterface.h"
@@ -32,7 +33,8 @@ UTexture2D* GetMediaSourceThumbnailTexture()
 	return LoadObject<UTexture2D>(GetTransientPackage(), MediaSourceThumbnailTexturePath);
 }
 
-const FLazyName FMediaSourceImageViewer::ItemTypeName = TEXT("MediaSource");
+const FLazyName FMediaSourceImageViewer::ItemTypeName_Asset = TEXT("MediaAsset");
+const FLazyName FMediaSourceImageViewer::ItemTypeName_File = TEXT("MediaFile");
 
 bool FMediaSourceImageViewer::FFactory::SupportsAsset(const FAssetData& InAssetData) const
 {
@@ -73,7 +75,7 @@ TSharedPtr<FMediaImageViewer> FMediaSourceImageViewer::FFactory::CreateImageView
 {
 	if (UMediaSource* MediaSource = Cast<UMediaSource>(InObject))
 	{
-		return MakeShared<FMediaSourceImageViewer>(MediaSource);
+		return MakeShared<FMediaSourceImageViewer>(MediaSource, FText::FromName(MediaSource->GetFName()));
 	}
 
 	return nullptr;
@@ -83,7 +85,7 @@ TSharedPtr<FMediaViewerLibraryItem> FMediaSourceImageViewer::FFactory::CreateLib
 {
 	if (UMediaSource* MediaSource = Cast<UMediaSource>(InObject))
 	{
-		return MakeShared<FMediaSourceImageViewer::FItem>(
+		return MakeShared<FMediaSourceImageViewer::FAssetItem>(
 			FText::FromName(MediaSource->GetFName()),
 			FText::Format(LOCTEXT("FactoryToolTipFormat", "{0} [Media Source]"), FText::FromString(MediaSource->GetPathName())),
 			MediaSource->HasAnyFlags(EObjectFlags::RF_Transient) || MediaSource->IsIn(GetTransientPackage()),
@@ -96,52 +98,78 @@ TSharedPtr<FMediaViewerLibraryItem> FMediaSourceImageViewer::FFactory::CreateLib
 
 bool FMediaSourceImageViewer::FFactory::SupportsItemType(FName InItemType) const
 {
-	return InItemType == FMediaSourceImageViewer::ItemTypeName;
+	return InItemType == FMediaSourceImageViewer::ItemTypeName_Asset;
 }
 
 TSharedPtr<FMediaViewerLibraryItem> FMediaSourceImageViewer::FFactory::CreateLibraryItem(const FMediaViewerLibraryItem& InSavedItem) const
 {
-	return MakeShared<FMediaSourceImageViewer::FItem>(FPrivateToken(), InSavedItem);
+	return MakeShared<FMediaSourceImageViewer::FAssetItem>(FPrivateToken(), InSavedItem);
 }
 
-FMediaSourceImageViewer::FItem::FItem(const FText& InName, const FText& InToolTip, bool bInTransient, TNotNull<UMediaSource*> InMediaSource)
-	: FItem(FGuid::NewGuid(), InName, InToolTip, bInTransient, InMediaSource)
+FMediaSourceImageViewer::FItem::FItem(const FText& InName, const FText& InToolTip, bool bInTransient, const FString& InStringValue)
+	: FMediaSourceImageViewer::FItem::FItem(FGuid::NewGuid(), InName, InToolTip, bInTransient, InStringValue)
 {
+
 }
 
-FMediaSourceImageViewer::FItem::FItem(const FGuid& InId, const FText& InName, const FText& InToolTip, bool bInTransient,
-	TNotNull<UMediaSource*> InMediaSource)
-	: FMediaViewerLibraryItem(InId, InName, InToolTip, bInTransient, InMediaSource->GetPathName())
+FMediaSourceImageViewer::FItem::FItem(const FGuid& InId, const FText& InName, const FText& InToolTip, bool bInTransient, const FString& InStringValue)
+	: FMediaViewerLibraryItem(InName, InToolTip, bInTransient, InStringValue)
 {
+
 }
 
-FMediaSourceImageViewer::FItem::FItem(FPrivateToken&& InPrivateToken, const FMediaViewerLibraryItem& InItem)
-	: FMediaViewerLibraryItem(InItem.GetId(), InItem.Name, InItem.ToolTip, InItem.IsTransient(), InItem.GetStringValue())
+TSharedPtr<FSlateBrush> FMediaSourceImageViewer::FItem::CreateThumbnail()
 {
+	// Loading the media to create a thumbnail is not a good idea.
+	TSharedRef<FSlateImageBrush> ThumbnailBrush = MakeShared<FSlateImageBrush>(
+		GetMediaSourceThumbnailTexture(),
+		FVector2D(64.0, 64.0)
+	);
+
+	return ThumbnailBrush;
 }
 
-FName FMediaSourceImageViewer::FItem::GetItemType() const
+FString FMediaSourceImageViewer::FItem::GetReferencerName() const
 {
-	return FMediaSourceImageViewer::ItemTypeName;
+	static const FString ReferencerName = TEXT("FMediaSourceImageViewer::FItem");
+	return ReferencerName;
 }
 
-FText FMediaSourceImageViewer::FItem::GetItemTypeDisplayName() const
+void FMediaSourceImageViewer::FItem::AddReferencedObjects(FReferenceCollector& Collector)
 {
-	return LOCTEXT("MediaSource", "Media");
+	if (Texture)
+	{
+		Collector.AddReferencedObject(Texture);
+	}
 }
 
-TSharedPtr<FSlateBrush> FMediaSourceImageViewer::FItem::CreateThumbnail()
+FMediaSourceImageViewer::FAssetItem::FAssetItem(const FText& InName, const FText& InToolTip, bool bInTransient, TNotNull<UMediaSource*> InMediaSource)
+	: FMediaSourceImageViewer::FAssetItem(FGuid::NewGuid(), InName, InToolTip, bInTransient, InMediaSource)
 {
-	// Loading the media to create a thumbnail is not a good idea.
-	TSharedRef<FSlateImageBrush> ThumbnailBrush = MakeShared<FSlateImageBrush>(
-		GetMediaSourceThumbnailTexture(), 
-		FVector2D(64.0, 64.0)
-	);
+}
 
-	return ThumbnailBrush;
+FMediaSourceImageViewer::FAssetItem::FAssetItem(const FGuid& InId, const FText& InName, const FText& InToolTip, bool bInTransient,
+	TNotNull<UMediaSource*> InMediaSource)
+	: FItem(InId, InName, InToolTip, bInTransient, InMediaSource->GetPathName())
+{
+}
+
+FMediaSourceImageViewer::FAssetItem::FAssetItem(FPrivateToken&& InPrivateToken, const FMediaViewerLibraryItem& InItem)
+	: FMediaSourceImageViewer::FItem(InItem.GetId(), InItem.Name, InItem.ToolTip, InItem.IsTransient(), InItem.GetStringValue())
+{
+}
+
+FName FMediaSourceImageViewer::FAssetItem::GetItemType() const
+{
+	return FMediaSourceImageViewer::ItemTypeName_Asset;
 }
 
-TSharedPtr<FMediaImageViewer> FMediaSourceImageViewer::FItem::CreateImageViewer() const
+FText FMediaSourceImageViewer::FAssetItem::GetItemTypeDisplayName() const
+{
+	return LOCTEXT("MediaSource", "Media Asset");
+}
+
+TSharedPtr<FMediaImageViewer> FMediaSourceImageViewer::FAssetItem::CreateImageViewer() const
 {
 	UMediaSource* MediaSource = LoadAssetFromString<UMediaSource>(StringValue);
 
@@ -150,49 +178,87 @@ TSharedPtr<FMediaImageViewer> FMediaSourceImageViewer::FItem::CreateImageViewer(
 		return nullptr;
 	}
 
+	const FText DisplayName = FText::FromName(MediaSource->GetFName());
+
 	if (Id.IsValid())
 	{
-		return MakeShared<FMediaSourceImageViewer>(Id, MediaSource);
+		return MakeShared<FMediaSourceImageViewer>(Id, MediaSource, DisplayName);
 	}
 
-	return MakeShared<FMediaSourceImageViewer>(MediaSource);
+	return MakeShared<FMediaSourceImageViewer>(MediaSource, DisplayName);
 }
 
-TSharedPtr<FMediaViewerLibraryItem> FMediaSourceImageViewer::FItem::Clone() const
+TSharedPtr<FMediaViewerLibraryItem> FMediaSourceImageViewer::FAssetItem::Clone() const
 {
 	if (StringValue.IsEmpty())
 	{
 		return nullptr;
 	}
 
-	return MakeShared<FMediaSourceImageViewer::FItem>(FPrivateToken(), *this);
+	return MakeShared<FMediaSourceImageViewer::FAssetItem>(FPrivateToken(), *this);
 }
 
-FString FMediaSourceImageViewer::FItem::GetReferencerName() const
+FMediaSourceImageViewer::FExternalItem::FExternalItem(const FText& InName, const FText& InToolTip, const FString& InFilePath)
+	: FMediaSourceImageViewer::FExternalItem(FGuid::NewGuid(), InName, InToolTip, InFilePath)
 {
-	static const FString ReferencerName = TEXT("FMediaSourceImageViewer::FItem");
-	return ReferencerName;
 }
 
-void FMediaSourceImageViewer::FItem::AddReferencedObjects(FReferenceCollector& Collector)
+FMediaSourceImageViewer::FExternalItem::FExternalItem(const FGuid& InId, const FText& InName, const FText& InToolTip, const FString& InFilePath)
+	: FMediaSourceImageViewer::FItem(InId, InName, InToolTip, /* Transient */ false, InFilePath)
 {
-	if (Texture)
+}
+
+FMediaSourceImageViewer::FExternalItem::FExternalItem(FPrivateToken&& InPrivateToken, const FMediaViewerLibraryItem& InItem)
+	: FMediaSourceImageViewer::FItem(InItem.GetId(), InItem.Name, InItem.ToolTip, InItem.IsTransient(), InItem.GetStringValue())
+{
+}
+
+FName FMediaSourceImageViewer::FExternalItem::GetItemType() const
+{
+	return FMediaSourceImageViewer::ItemTypeName_File;
+}
+
+FText FMediaSourceImageViewer::FExternalItem::GetItemTypeDisplayName() const
+{
+	return LOCTEXT("MediaFile", "Media File");
+}
+
+TSharedPtr<FMediaImageViewer> FMediaSourceImageViewer::FExternalItem::CreateImageViewer() const
+{
+	UFileMediaSource* FileMediaSource = NewObject<UFileMediaSource>(GetTransientPackage());
+	FileMediaSource->SetFilePath(StringValue);
+
+	const FText DisplayName = FText::FromString(StringValue);
+
+	if (Id.IsValid())
 	{
-		Collector.AddReferencedObject(Texture);
+		return MakeShared<FMediaSourceImageViewer>(Id, FileMediaSource, DisplayName);
+	}
+
+	return MakeShared<FMediaSourceImageViewer>(FileMediaSource, DisplayName);
+}
+
+TSharedPtr<FMediaViewerLibraryItem> FMediaSourceImageViewer::FExternalItem::Clone() const
+{
+	if (StringValue.IsEmpty())
+	{
+		return nullptr;
 	}
+
+	return MakeShared<FMediaSourceImageViewer::FExternalItem>(FPrivateToken(), *this);
 }
 
-FMediaSourceImageViewer::FMediaSourceImageViewer(TNotNull<UMediaSource*> InMediaSource)
-	: FMediaSourceImageViewer(FGuid::NewGuid(), InMediaSource)
+FMediaSourceImageViewer::FMediaSourceImageViewer(TNotNull<UMediaSource*> InMediaSource, const FText& InDisplayName)
+	: FMediaSourceImageViewer(FGuid::NewGuid(), InMediaSource, InDisplayName)
 {
 }
 
-FMediaSourceImageViewer::FMediaSourceImageViewer(const FGuid& InId, TNotNull<UMediaSource*> InMediaSource)
+FMediaSourceImageViewer::FMediaSourceImageViewer(const FGuid& InId, TNotNull<UMediaSource*> InMediaSource, const FText& InDisplayName)
 	: FMediaImageViewer({
 		InId,
 		FIntPoint(1, 1),
 		1,
-		FText::FromName(InMediaSource->GetFName())
+		InDisplayName
 	})
 {
 	MediaSourceSettings.MediaSource = InMediaSource;
@@ -244,12 +310,25 @@ TSharedPtr<FMediaViewerLibraryItem> FMediaSourceImageViewer::CreateLibraryItem()
 		return nullptr;
 	}
 
-	return MakeShared<FItem>(
+	const bool bIsTransient = MediaSource->HasAnyFlags(EObjectFlags::RF_Transient) || MediaSource->IsIn(GetTransientPackage());
+	UFileMediaSource* FileMediaSource = Cast<UFileMediaSource>(MediaSource);
+
+	if (!bIsTransient || !FileMediaSource)
+	{
+		return MakeShared<FAssetItem>(
+			ImageInfo.Id,
+			FText::FromName(MediaSource->GetFName()),
+			FText::Format(LOCTEXT("LibraryToolTipFormat", "{0} [Media Texture - {1}x{2}]"), FText::FromString(MediaSource->GetPathName()), FText::AsNumber(ImageInfo.Size.X), FText::AsNumber(ImageInfo.Size.Y)),
+			bIsTransient,
+			MediaSource
+		);
+	}
+
+	return MakeShared<FExternalItem>(
 		ImageInfo.Id,
-		FText::FromName(MediaSource->GetFName()),
-		FText::Format(LOCTEXT("LibraryToolTipFormat", "{0} [Media Texture - {1}x{2}]"), FText::FromString(MediaSource->GetPathName()), FText::AsNumber(ImageInfo.Size.X), FText::AsNumber(ImageInfo.Size.Y)),
-		MediaSource->HasAnyFlags(EObjectFlags::RF_Transient) || MediaSource->IsIn(GetTransientPackage()),
-		MediaSource
+		FText::FromString(FileMediaSource->GetFilePath()),
+		FText::Format(LOCTEXT("LibraryToolTipFormat", "{0} [Media Texture - {1}x{2}]"), FText::FromString(FileMediaSource->GetFilePath()), FText::AsNumber(ImageInfo.Size.X), FText::AsNumber(ImageInfo.Size.Y)),
+		FileMediaSource->GetFilePath()
 	);
 }
 
diff --git a/Engine/Plugins/Experimental/MediaViewer/Source/MediaViewer/Private/ImageViewers/MediaSourceImageViewer.h b/Engine/Plugins/Experimental/MediaViewer/Source/MediaViewer/Private/ImageViewers/MediaSourceImageViewer.h
--- a/Engine/Plugins/Experimental/MediaViewer/Source/MediaViewer/Private/ImageViewers/MediaSourceImageViewer.h
+++ b/Engine/Plugins/Experimental/MediaViewer/Source/MediaViewer/Private/ImageViewers/MediaSourceImageViewer.h
@@ -56,32 +56,61 @@ public:
 	{
 		TObjectPtr<UTexture> Texture;
 
-		FItem(const FText& InName, const FText& InToolTip, bool bInTransient, TNotNull<UMediaSource*> InMediaSource);
+		FItem(const FText& InName, const FText& InToolTip, bool bInTransient, const FString& InStringValue);
 
-		FItem(const FGuid& InId, const FText& InName, const FText& InToolTip, bool bInTransient, TNotNull<UMediaSource*> InMediaSource);
+		FItem(const FGuid& InId, const FText& InName, const FText& InToolTip, bool bInTransient, const FString& InStringValue);
 
-		FItem(FPrivateToken&& InPrivateToken, const FMediaViewerLibraryItem& InItem);
+		//~ Begin FMediaViewerLibraryItem
+		virtual TSharedPtr<FSlateBrush> CreateThumbnail() override;
+		//~ End FMediaViewerLibraryItem
 
-		virtual ~FItem() override = default;
+		//~ Begin FGCObject
+		virtual FString GetReferencerName() const override;
+		virtual void AddReferencedObjects(FReferenceCollector& Collector) override;
+		//~ End FGCObject
+	};
+
+	struct FAssetItem : public FItem
+	{
+		FAssetItem(const FText& InName, const FText& InToolTip, bool bInTransient, TNotNull<UMediaSource*> InMediaSource);
+
+		FAssetItem(const FGuid& InId, const FText& InName, const FText& InToolTip, bool bInTransient, TNotNull<UMediaSource*> InMediaSource);
+
+		FAssetItem(FPrivateToken&& InPrivateToken, const FMediaViewerLibraryItem& InItem);
+
+		virtual ~FAssetItem() override = default;
 
 		//~ Begin FMediaViewerLibraryItem
 		virtual FName GetItemType() const override;
 		virtual FText GetItemTypeDisplayName() const override;
-		virtual TSharedPtr<FSlateBrush> CreateThumbnail() override;
 		virtual TSharedPtr<FMediaImageViewer> CreateImageViewer() const override;
 		virtual TSharedPtr<FMediaViewerLibraryItem> Clone() const override;
 		//~ End FMediaViewerLibraryItem
+	};
 
-		//~ Begin FGCObject
-		virtual FString GetReferencerName() const override;
-		virtual void AddReferencedObjects(FReferenceCollector& Collector) override;
-		//~ End FGCObject
+	struct FExternalItem : public FItem
+	{
+		FExternalItem(const FText& InName, const FText& InToolTip, const FString& InFilePath);
+
+		FExternalItem(const FGuid& InId, const FText& InName, const FText& InToolTip, const FString& InFilePath);
+
+		FExternalItem(FPrivateToken&& InPrivateToken, const FMediaViewerLibraryItem& InItem);
+
+		virtual ~FExternalItem() override = default;
+
+		//~ Begin FMediaViewerLibraryItem
+		virtual FName GetItemType() const override;
+		virtual FText GetItemTypeDisplayName() const override;
+		virtual TSharedPtr<FMediaImageViewer> CreateImageViewer() const override;
+		virtual TSharedPtr<FMediaViewerLibraryItem> Clone() const override;
+		//~ End FMediaViewerLibraryItem
 	};
 
-	static const FLazyName ItemTypeName;
+	static const FLazyName ItemTypeName_Asset;
+	static const FLazyName ItemTypeName_File;
 
-	FMediaSourceImageViewer(TNotNull<UMediaSource*> InMediaSource);
-	FMediaSourceImageViewer(const FGuid& InId, TNotNull<UMediaSource*> InMediaSource);
+	FMediaSourceImageViewer(TNotNull<UMediaSource*> InMediaSource, const FText& InDisplayName);
+	FMediaSourceImageViewer(const FGuid& InId, TNotNull<UMediaSource*> InMediaSource, const FText& InDisplayName);
 
 	virtual ~FMediaSourceImageViewer() override = default;
 
diff --git a/Engine/Plugins/Experimental/MediaViewer/Source/MediaViewer/Private/MediaViewerModule.cpp b/Engine/Plugins/Experimental/MediaViewer/Source/MediaViewer/Private/MediaViewerModule.cpp
--- a/Engine/Plugins/Experimental/MediaViewer/Source/MediaViewer/Private/MediaViewerModule.cpp
+++ b/Engine/Plugins/Experimental/MediaViewer/Source/MediaViewer/Private/MediaViewerModule.cpp
@@ -475,7 +475,7 @@ void FMediaViewerModule::RegisterDefaultImageViewers()
 {
 	RegisterFactory(FColorImageViewer::ItemTypeName,                 MakeShared<FColorImageViewer::FFactory>());
 	RegisterFactory(FMaterialInterfaceImageViewer::ItemTypeName,     MakeShared<FMaterialInterfaceImageViewer::FFactory>());
-	RegisterFactory(FMediaSourceImageViewer::ItemTypeName,           MakeShared<FMediaSourceImageViewer::FFactory>());
+	RegisterFactory(FMediaSourceImageViewer::ItemTypeName_Asset,     MakeShared<FMediaSourceImageViewer::FFactory>());
 	RegisterFactory(FMediaTextureImageViewer::ItemTypeName,          MakeShared<FMediaTextureImageViewer::FFactory>());
 	RegisterFactory(FTexture2DImageViewer::ItemTypeName,             MakeShared<FTexture2DImageViewer::FFactory>());
 	RegisterFactory(FTextureRenderTarget2DImageViewer::ItemTypeName, MakeShared<FTextureRenderTarget2DImageViewer::FFactory>());
diff --git a/Engine/Plugins/Experimental/MediaViewer/Source/MediaViewer/Private/Widgets/SMediaViewerDropTarget.cpp b/Engine/Plugins/Experimental/MediaViewer/Source/MediaViewer/Private/Widgets/SMediaViewerDropTarget.cpp
--- a/Engine/Plugins/Experimental/MediaViewer/Source/MediaViewer/Private/Widgets/SMediaViewerDropTarget.cpp
+++ b/Engine/Plugins/Experimental/MediaViewer/Source/MediaViewer/Private/Widgets/SMediaViewerDropTarget.cpp
@@ -204,7 +204,7 @@ void SMediaViewerDropTarget::HandleDroppeFileOp(const FExternalDragOperation& In
 	{
 		Delegates->SetImageViewer.Execute(
 			static_cast<EMediaImageViewerPosition>(Index),
-			MakeShared<FMediaSourceImageViewer>(Sources[Index])
+			MakeShared<FMediaSourceImageViewer>(Sources[Index], FText::FromString(Sources[Index]->GetFilePath()))
 		);
 	}
 }
