Media Viewer: Group update
- History is now the second group
- Moved media texture dynamic group below the editor viewport group.

diff --git a/Engine/Plugins/Experimental/MediaViewer/Source/MediaViewer/Private/Library/MediaViewerLibrary.cpp b/Engine/Plugins/Experimental/MediaViewer/Source/MediaViewer/Private/Library/MediaViewerLibrary.cpp
--- a/Engine/Plugins/Experimental/MediaViewer/Source/MediaViewer/Private/Library/MediaViewerLibrary.cpp
+++ b/Engine/Plugins/Experimental/MediaViewer/Source/MediaViewer/Private/Library/MediaViewerLibrary.cpp
@@ -45,10 +45,15 @@ FMediaViewerLibrary::FMediaViewerLibrary(const FArgs& InArgs)
 		/* Dynamic */ false
 	));
 
+
+
 	Groups.Add(MakeShared<FMediaViewerLibraryGroup>(
 		GetHistoryGroupId(),
 		LOCTEXT("History", "History"),
-		LOCTEXT("HistoryTooltip", "Up to the last 25 viewed items."),
+		FText::Format(
+			LOCTEXT("HistoryTooltip", "Up to the last {0} viewed items."),
+			FText::AsNumber(UE::MediaViewer::Private::MaxHistoryEntries)
+		),
 		/* Dynamic */ false
 	));
 }
@@ -79,8 +84,7 @@ bool FMediaViewerLibrary::AddGroup(const TSharedRef<FMediaViewerLibraryGroup>& I
 		return false;
 	}
 
-	// Make sure history is always the last entry
-	Groups.Insert(InNewGroup, Groups.Num() - 1);
+	Groups.Add(InNewGroup);;
 
 	OnChanged(IMediaViewerLibrary::EChangeType::GroupAdded);
 
diff --git a/Engine/Plugins/Experimental/MediaViewer/Source/MediaViewer/Private/Widgets/SMediaViewerLibraryPrivate.cpp b/Engine/Plugins/Experimental/MediaViewer/Source/MediaViewer/Private/Widgets/SMediaViewerLibraryPrivate.cpp
--- a/Engine/Plugins/Experimental/MediaViewer/Source/MediaViewer/Private/Widgets/SMediaViewerLibraryPrivate.cpp
+++ b/Engine/Plugins/Experimental/MediaViewer/Source/MediaViewer/Private/Widgets/SMediaViewerLibraryPrivate.cpp
@@ -18,11 +18,6 @@
 
 #define LOCTEXT_NAMESPACE "SMediaViewerLibraryPrivate"
 
-namespace UE::MediaViewer::Private
-{
-	constexpr int32 MaxHistoryEntries = 25;
-}
-
 namespace UE::MediaViewer::Private
 {
 
@@ -46,8 +41,8 @@ void SMediaViewerLibraryPrivate::Construct(const FArguments& InArgs, const TShar
 
 	TSharedRef<FMediaViewerLibrary> LibraryImpl = GetLibrary();
 
-	LibraryImpl->AddGroup(MakeShared<FMediaTextureGroup>(LibraryImpl));
 	LibraryImpl->AddGroup(MakeShared<FLevelEditorViewportGroup>(LibraryImpl));
+	LibraryImpl->AddGroup(MakeShared<FMediaTextureGroup>(LibraryImpl));
 }
 
 TSharedRef<FMediaViewerLibrary> SMediaViewerLibraryPrivate::GetLibrary() const
