Horde issue 876640

diff --git a/Engine/Plugins/Experimental/MediaViewer/Source/MediaViewer/Private/Library/MediaViewerLibrary.h b/Engine/Plugins/Experimental/MediaViewer/Source/MediaViewer/Private/Library/MediaViewerLibrary.h
--- a/Engine/Plugins/Experimental/MediaViewer/Source/MediaViewer/Private/Library/MediaViewerLibrary.h
+++ b/Engine/Plugins/Experimental/MediaViewer/Source/MediaViewer/Private/Library/MediaViewerLibrary.h
@@ -15,6 +15,8 @@ struct FMediaViewerLibraryItem;
 namespace UE::MediaViewer::Private
 {
 
+constexpr int32 MaxHistoryEntries = 20;
+
 /**
  * Implementation of @see IMediaViewerLibrary.
  */
