Media Viewer: Changed overlay / media controls
- "Status bar" is now at the top.
- Media controls now have a different style and have moved down.

diff --git a/Engine/Plugins/Experimental/MediaViewer/Source/MediaViewer/Private/MediaViewerStyle.cpp b/Engine/Plugins/Experimental/MediaViewer/Source/MediaViewer/Private/MediaViewerStyle.cpp
--- a/Engine/Plugins/Experimental/MediaViewer/Source/MediaViewer/Private/MediaViewerStyle.cpp
+++ b/Engine/Plugins/Experimental/MediaViewer/Source/MediaViewer/Private/MediaViewerStyle.cpp
@@ -44,6 +44,14 @@ FMediaViewerStyle::FMediaViewerStyle()
 	ToolbarButtonStyle.SetPressedPadding(FMargin(0.f, 0.f, 0.f, 0.f));
 	Set("ToolbarButtonStyle", ToolbarButtonStyle);
 
+	Set("MediaButtons", FButtonStyle(FAppStyle::Get().GetWidgetStyle<FButtonStyle>("Animation.PlayControlsButton"))
+		.SetNormal(FSlateNoResource())
+		.SetDisabled(FSlateNoResource())
+		.SetHovered(FSlateRoundedBoxBrush(FLinearColor(0.2, 0.2, 0.2, 0.5), 3.f, FVector2D(20.f)))
+		.SetPressed(FSlateRoundedBoxBrush(FLinearColor(0.1, 0.1, 0.1, 0.5), 3.f, FVector2D(20.f)))
+		.SetNormalPadding(FMargin(2.f, 2.f, 2.f, 2.f))
+		.SetPressedPadding(FMargin(2.f, 2.f, 2.f, 2.f)));
+
 	FSlateStyleRegistry::RegisterSlateStyle(*this);
 }
 
diff --git a/Engine/Plugins/Experimental/MediaViewer/Source/MediaViewer/Private/Widgets/SMediaImageViewerOverlay.cpp b/Engine/Plugins/Experimental/MediaViewer/Source/MediaViewer/Private/Widgets/SMediaImageViewerOverlay.cpp
--- a/Engine/Plugins/Experimental/MediaViewer/Source/MediaViewer/Private/Widgets/SMediaImageViewerOverlay.cpp
+++ b/Engine/Plugins/Experimental/MediaViewer/Source/MediaViewer/Private/Widgets/SMediaImageViewerOverlay.cpp
@@ -74,20 +74,20 @@ void SMediaImageViewerOverlay::Construct(const FArguments& InArgs, EMediaImageVi
 			SAssignNew(Container, SVerticalBox)
 
 			+ SVerticalBox::Slot()
-			.FillHeight(1.f)
+			.AutoHeight()
 			.HAlign(HAlign_Fill)
-			.VAlign(VAlign_Fill)
+			.VAlign(VAlign_Center)
+			.Padding(PaddingLeft, 0.f, PaddingRight, 0.f)
 			[
-				CreateOverlay(InArgs._bShowPanelName)
+				CreateStatusBar(InDelegates)
 			]
 
 			+ SVerticalBox::Slot()
-			.AutoHeight()
+			.FillHeight(1.f)
 			.HAlign(HAlign_Fill)
-			.VAlign(VAlign_Center)
-			.Padding(PaddingLeft, 0.f, PaddingRight, 0.f)
+			.VAlign(VAlign_Fill)
 			[
-				CreateStatusBar(InDelegates)
+				CreateOverlay(InArgs._bShowPanelName)
 			]
 		]
 	];
diff --git a/Engine/Plugins/Experimental/MediaViewer/Source/MediaViewer/Private/Widgets/SMediaSourceOverlay.cpp b/Engine/Plugins/Experimental/MediaViewer/Source/MediaViewer/Private/Widgets/SMediaSourceOverlay.cpp
--- a/Engine/Plugins/Experimental/MediaViewer/Source/MediaViewer/Private/Widgets/SMediaSourceOverlay.cpp
+++ b/Engine/Plugins/Experimental/MediaViewer/Source/MediaViewer/Private/Widgets/SMediaSourceOverlay.cpp
@@ -8,6 +8,7 @@
 #include "MediaPlayer.h"
 #include "MediaPlayerEditorModule.h"
 #include "MediaStream.h"
+#include "MediaViewerStyle.h"
 #include "Misc/FrameRate.h"
 #include "Misc/Timecode.h"
 #include "Misc/Timespan.h"
@@ -17,6 +18,7 @@
 #include "Widgets/Colors/SColorBlock.h"
 #include "Widgets/Images/SImage.h"
 #include "Widgets/Input/SButton.h"
+#include "Widgets/Layout/SBorder.h"
 #include "Widgets/SBoxPanel.h"
 #include "Widgets/SNullWidget.h"
 #include "Widgets/Text/STextBlock.h"
@@ -38,22 +40,24 @@ void SMediaSourceOverlay::Construct(const FArguments& InArgs, const TSharedRef<F
 
 	ChildSlot
 	[
-		SNew(SVerticalBox)
-		
-		+ SVerticalBox::Slot()
-		.AutoHeight()
-		.HAlign(HAlign_Fill)
-		.MaxSize(400.f)
-		[
-			CreateSlider()
-		]
-
-		+ SVerticalBox::Slot()
-		.AutoHeight()
-		.HAlign(HAlign_Fill)
-		.Padding(0.f, 5.f, 0.f, 0.f)
+		SNew(SBorder)
+		.Padding(5.f)
+		.BorderImage(FAppStyle::GetBrush("ToolTip.Background"))
+		.BorderBackgroundColor(FLinearColor(1.f, 1.f, 1.f, 0.75f))
 		[
-			CreateControls()
+			SNew(SVerticalBox)		
+			+ SVerticalBox::Slot()
+			.AutoHeight()
+			[
+				CreateSlider()
+			]
+			+ SVerticalBox::Slot()
+			.AutoHeight()
+			.HAlign(HAlign_Fill)
+			.Padding(0.f, 0.f, 0.f, 0.f)
+			[
+				CreateControls()
+			]
 		]
 	];
 }
@@ -126,18 +130,7 @@ TSharedRef<SWidget> SMediaSourceOverlay::CreateSlider()
 	MediaPlayerSlider->SetSliderHandleColor(FSlateColor(EStyleColor::AccentBlue));
 	MediaPlayerSlider->SetVisibleWhenInactive(EVisibility::Visible);
 
-	return SNew(SOverlay)
-		+ SOverlay::Slot()
-		.Padding(5.f, 11.f, 5.f, 11.f)
-		[
-			SNew(SColorBlock)
-			.Color(FStyleColors::Header.GetSpecifiedColor())
-			.Size(FVector2D(10.f, 2.f))
-		]
-		+ SOverlay::Slot()
-		[
-			MediaPlayerSlider
-		];
+	return MediaPlayerSlider;
 }
 
 TSharedRef<SWidget> SMediaSourceOverlay::CreateControls()
@@ -195,11 +188,34 @@ TSharedRef<SWidget> SMediaSourceOverlay::CreateControls()
 			.VAlign(VAlign_Center)
 			.IsEnabled(this, &SMediaSourceOverlay::Rewind_IsEnabled)
 			.OnClicked(this, &SMediaSourceOverlay::Rewind_OnClicked)
+			.ButtonStyle(FMediaViewerStyle::Get(), "MediaButtons")
+			.ContentPadding(2.f)
 			[
 				SNew(SImage)
 				.ColorAndOpacity(FSlateColor::UseForeground())
-				.Image(FSlateIconFinder::FindIcon("MediaPlateEditor.RewindMedia.Small").GetIcon())
+				.Image(FSlateIcon(FAppStyle::GetAppStyleSetName(), "Animation.Backward_End").GetIcon())
 				.ToolTipText(LOCTEXT("Rewind", "Rewind the media to the beginning"))
+				.DesiredSizeOverride(FVector2D(20.f))
+			]
+		]
+
+		// Step back button.
+		+ SHorizontalBox::Slot()
+		.VAlign(VAlign_Center)
+		.AutoWidth()
+		[
+			SNew(SButton)
+			.VAlign(VAlign_Center)
+			.IsEnabled(this, &SMediaSourceOverlay::StepBack_IsEnabled)
+			.OnClicked(this, &SMediaSourceOverlay::StepBack_OnClicked)
+			.ButtonStyle(FMediaViewerStyle::Get(), "MediaButtons")
+			.ContentPadding(2.f)
+			[
+				SNew(SImage)
+				.ColorAndOpacity(FSlateColor::UseForeground())
+				.Image(FSlateIcon(FAppStyle::GetAppStyleSetName(), "Animation.Backward_Step").GetIcon())
+				.ToolTipText(LOCTEXT("StepBack", "Step back 1 frame.\n\nOnly available while paused."))
+				.DesiredSizeOverride(FVector2D(20.f))
 			]
 		]
 
@@ -212,11 +228,14 @@ TSharedRef<SWidget> SMediaSourceOverlay::CreateControls()
 			.VAlign(VAlign_Center)
 			.IsEnabled(this, &SMediaSourceOverlay::Reverse_IsEnabled)
 			.OnClicked(this, &SMediaSourceOverlay::Reverse_OnClicked)
+			.ButtonStyle(FMediaViewerStyle::Get(), "MediaButtons")
+			.ContentPadding(2.f)
 			[
 				SNew(SImage)
 				.ColorAndOpacity(FSlateColor::UseForeground())
-				.Image(FSlateIconFinder::FindIcon("MediaPlateEditor.ReverseMedia.Small").GetIcon())
-				.ToolTipText(LOCTEXT("Reverse", "Play media in reverse."))
+				.Image(FSlateIcon(FAppStyle::GetAppStyleSetName(), "Animation.Backward").GetIcon())
+				.ToolTipText(LOCTEXT("Reverse", "Play media in reverse.\n\nNot widely supported by media decoders."))
+				.DesiredSizeOverride(FVector2D(20.f))
 			]
 		]
 
@@ -230,11 +249,14 @@ TSharedRef<SWidget> SMediaSourceOverlay::CreateControls()
 			.IsEnabled(this, &SMediaSourceOverlay::Play_IsEnabled)
 			.Visibility(this, &SMediaSourceOverlay::Play_GetVisibility)
 			.OnClicked(this, &SMediaSourceOverlay::Play_OnClicked)
+			.ButtonStyle(FMediaViewerStyle::Get(), "MediaButtons")
+			.ContentPadding(2.f)
 			[
 				SNew(SImage)
 				.ColorAndOpacity(FSlateColor::UseForeground())
-				.Image(FSlateIconFinder::FindIcon("MediaPlateEditor.PlayMedia.Small").GetIcon())
+				.Image(FSlateIcon(FAppStyle::GetAppStyleSetName(), "Animation.Forward").GetIcon())
 				.ToolTipText(LOCTEXT("Play", "Start media playback"))
+				.DesiredSizeOverride(FVector2D(20.f))
 			]
 		]
 
@@ -248,15 +270,38 @@ TSharedRef<SWidget> SMediaSourceOverlay::CreateControls()
 			.IsEnabled(this, &SMediaSourceOverlay::Pause_IsEnabled)
 			.Visibility(this, &SMediaSourceOverlay::Pause_GetVisibility)
 			.OnClicked(this, &SMediaSourceOverlay::Pause_OnClicked)
+			.ButtonStyle(FMediaViewerStyle::Get(), "MediaButtons")
+			.ContentPadding(2.f)
 			[
 				SNew(SImage)
 				.ColorAndOpacity(FSlateColor::UseForeground())
-				.Image(FSlateIconFinder::FindIcon("MediaPlateEditor.PauseMedia.Small").GetIcon())
+				.Image(FSlateIcon(FAppStyle::GetAppStyleSetName(), "Animation.Pause").GetIcon())
 				.ToolTipText(LOCTEXT("Pause", "Pause media playback"))
+				.DesiredSizeOverride(FVector2D(20.f))
+			]
+		]
+
+		// Step forward button.
+		+ SHorizontalBox::Slot()
+		.VAlign(VAlign_Center)
+		.AutoWidth()
+		[
+			SNew(SButton)
+			.VAlign(VAlign_Center)
+			.IsEnabled(this, &SMediaSourceOverlay::StepForward_IsEnabled)
+			.OnClicked(this, &SMediaSourceOverlay::StepForward_OnClicked)
+			.ButtonStyle(FMediaViewerStyle::Get(), "MediaButtons")
+			.ContentPadding(2.f)
+			[
+				SNew(SImage)
+				.ColorAndOpacity(FSlateColor::UseForeground())
+				.Image(FSlateIcon(FAppStyle::GetAppStyleSetName(), "Animation.Forward_Step").GetIcon())
+				.ToolTipText(LOCTEXT("StepForward", "Step forward 1 frame.\n\nOnly available while paused."))
+				.DesiredSizeOverride(FVector2D(20.f))
 			]
 		]
 
-		// Forward button.
+		// Fast Forward button.
 		+ SHorizontalBox::Slot()
 		.VAlign(VAlign_Center)
 		.AutoWidth()
@@ -265,13 +310,14 @@ TSharedRef<SWidget> SMediaSourceOverlay::CreateControls()
 			.VAlign(VAlign_Center)
 			.IsEnabled(this, &SMediaSourceOverlay::Forward_IsEnabled)
 			.OnClicked(this, &SMediaSourceOverlay::Forward_OnClicked)
+			.ButtonStyle(FMediaViewerStyle::Get(), "MediaButtons")
+			.ContentPadding(2.f)
 			[
 				SNew(SImage)
 				.ColorAndOpacity(FSlateColor::UseForeground())
-				.Image(FSlateIconFinder::FindIcon("MediaPlateEditor.RewindMedia.Small").GetIcon())
-				// There is no correct icon pointing forward, so this rotates the rewind icon 180 degrees and adds the required offset.
-				.RenderTransform(FSlateRenderTransform(TTransform2<float>(TQuat2<float>(PI), FVector2f(20.f, 20.f))))
+				.Image(FSlateIcon(FAppStyle::GetAppStyleSetName(), "Animation.Forward_End").GetIcon())
 				.ToolTipText(LOCTEXT("Forward", "Fast forward the media to the end."))
+				.DesiredSizeOverride(FVector2D(20.f))
 			]
 		]
 
@@ -453,6 +499,33 @@ FReply SMediaSourceOverlay::Reverse_OnClicked()
 	return FReply::Handled();
 }
 
+bool SMediaSourceOverlay::StepBack_IsEnabled() const
+{
+	if (UMediaPlayer* MediaPlayer = GetMediaPlayer())
+	{
+		const float Rate = FMath::Abs(MediaPlayer->GetRate());
+
+		return MediaPlayer->IsReady() && MediaPlayer->IsPaused();
+	}
+
+	return false;
+}
+
+FReply SMediaSourceOverlay::StepBack_OnClicked()
+{
+	if (UMediaPlayer* MediaPlayer = GetMediaPlayer())
+	{
+		const float FrameRate = MediaPlayer->GetVideoTrackFrameRate(INDEX_NONE, INDEX_NONE);
+
+		if (!FMath::IsNearlyZero(FrameRate))
+		{
+			MediaPlayer->Seek(MediaPlayer->GetTime() - FTimespan::FromSeconds(1.f / FrameRate));
+		}
+	}
+
+	return FReply::Handled();
+}
+
 bool SMediaSourceOverlay::Play_IsEnabled() const
 {
 	if (UMediaPlayer* MediaPlayer = GetMediaPlayer())
@@ -520,6 +593,33 @@ FReply SMediaSourceOverlay::Pause_OnClicked()
 	return FReply::Handled();
 }
 
+bool SMediaSourceOverlay::StepForward_IsEnabled() const
+{
+	if (UMediaPlayer* MediaPlayer = GetMediaPlayer())
+	{
+		const float Rate = FMath::Abs(MediaPlayer->GetRate());
+
+		return MediaPlayer->IsReady() && MediaPlayer->IsPaused();
+	}
+
+	return false;
+}
+
+FReply SMediaSourceOverlay::StepForward_OnClicked()
+{
+	if (UMediaPlayer* MediaPlayer = GetMediaPlayer())
+	{
+		const float FrameRate = MediaPlayer->GetVideoTrackFrameRate(INDEX_NONE, INDEX_NONE);
+
+		if (!FMath::IsNearlyZero(FrameRate))
+		{
+			MediaPlayer->Seek(MediaPlayer->GetTime() + FTimespan::FromSeconds(1.f / FrameRate));
+		}
+	}
+
+	return FReply::Handled();
+}
+
 bool SMediaSourceOverlay::Forward_IsEnabled() const
 {
 	if (UMediaPlayer* MediaPlayer = GetMediaPlayer())
diff --git a/Engine/Plugins/Experimental/MediaViewer/Source/MediaViewer/Private/Widgets/SMediaSourceOverlay.h b/Engine/Plugins/Experimental/MediaViewer/Source/MediaViewer/Private/Widgets/SMediaSourceOverlay.h
--- a/Engine/Plugins/Experimental/MediaViewer/Source/MediaViewer/Private/Widgets/SMediaSourceOverlay.h
+++ b/Engine/Plugins/Experimental/MediaViewer/Source/MediaViewer/Private/Widgets/SMediaSourceOverlay.h
@@ -59,6 +59,9 @@ protected:
 	bool Reverse_IsEnabled() const;
 	FReply Reverse_OnClicked();
 
+	bool StepBack_IsEnabled() const;
+	FReply StepBack_OnClicked();
+
 	bool Play_IsEnabled() const;
 	EVisibility Play_GetVisibility() const;
 	FReply Play_OnClicked();
@@ -67,6 +70,9 @@ protected:
 	EVisibility Pause_GetVisibility() const;
 	FReply Pause_OnClicked();
 
+	bool StepForward_IsEnabled() const;
+	FReply StepForward_OnClicked();
+
 	bool Forward_IsEnabled() const;
 	FReply Forward_OnClicked();
 };
diff --git a/Engine/Plugins/Experimental/MediaViewer/Source/MediaViewer/Public/ImageViewer/MediaImageViewer.h b/Engine/Plugins/Experimental/MediaViewer/Source/MediaViewer/Public/ImageViewer/MediaImageViewer.h
--- a/Engine/Plugins/Experimental/MediaViewer/Source/MediaViewer/Public/ImageViewer/MediaImageViewer.h
+++ b/Engine/Plugins/Experimental/MediaViewer/Source/MediaViewer/Public/ImageViewer/MediaImageViewer.h
@@ -208,7 +208,7 @@ public:
 	virtual void ExtendToolbar(const TSharedRef<FExtender>& InOutToolbarExtender) {}
 
 	/**
-	 * Extend the status bar (bottom). @see MediaViewer.h
+	 * Extend the status bar (top over overlay). @see MediaViewer.h
 	 * @param InOutStatusBarExtender - The extender to populate.
 	 */
 	virtual void ExtendStatusBar(FMediaImageStatusBarExtender& InOutStatusBarExtender) {}
