Media Viewer: Revamped the single-view drop targets
- Now can go straight to AB view using 3 different drop targets.

diff --git a/Engine/Plugins/Experimental/MediaViewer/Source/MediaViewer/Private/Widgets/MediaViewerDelegates.h b/Engine/Plugins/Experimental/MediaViewer/Source/MediaViewer/Private/Widgets/MediaViewerDelegates.h
--- a/Engine/Plugins/Experimental/MediaViewer/Source/MediaViewer/Private/Widgets/MediaViewerDelegates.h
+++ b/Engine/Plugins/Experimental/MediaViewer/Source/MediaViewer/Private/Widgets/MediaViewerDelegates.h
@@ -54,6 +54,9 @@ struct FMediaViewerDelegates
 
 	/** Changes to single view mode. */
 	FSimpleDelegate SetSingleView;
+
+	/** Changes to AB view. */
+	FSimpleDelegate SetABView;
 	
 	/** Gets the orientation in A/B mode. Will return the currently selected orientation even if not in A/B mode. */
 	FGetOrientation GetABOrientation;
diff --git a/Engine/Plugins/Experimental/MediaViewer/Source/MediaViewer/Private/Widgets/SMediaImageViewerOverlay.cpp b/Engine/Plugins/Experimental/MediaViewer/Source/MediaViewer/Private/Widgets/SMediaImageViewerOverlay.cpp
--- a/Engine/Plugins/Experimental/MediaViewer/Source/MediaViewer/Private/Widgets/SMediaImageViewerOverlay.cpp
+++ b/Engine/Plugins/Experimental/MediaViewer/Source/MediaViewer/Private/Widgets/SMediaImageViewerOverlay.cpp
@@ -61,12 +61,6 @@ void SMediaImageViewerOverlay::Construct(const FArguments& InArgs, EMediaImageVi
 
 	Delegates = InDelegates;
 
-	const bool bDisplayLeft = InPosition == EMediaImageViewerPosition::First
-		&& InDelegates->GetABOrientation.Execute() == Orient_Horizontal;
-
-	const float PaddingLeft = bDisplayLeft ? 30.f : 0.f;
-	const float PaddingRight = !bDisplayLeft ? 30.f : 0.f;
-
 	ChildSlot
 	[
 		SNew(SScissorRectBox)
@@ -77,7 +71,6 @@ void SMediaImageViewerOverlay::Construct(const FArguments& InArgs, EMediaImageVi
 			.AutoHeight()
 			.HAlign(HAlign_Fill)
 			.VAlign(VAlign_Center)
-			.Padding(PaddingLeft, 0.f, PaddingRight, 0.f)
 			[
 				CreateStatusBar(InDelegates)
 			]
@@ -87,7 +80,7 @@ void SMediaImageViewerOverlay::Construct(const FArguments& InArgs, EMediaImageVi
 			.HAlign(HAlign_Fill)
 			.VAlign(VAlign_Fill)
 			[
-				CreateOverlay(InArgs._bShowPanelName)
+				CreateOverlay(InArgs._bComparisonView)
 			]
 		]
 	];
@@ -429,7 +422,7 @@ void SMediaImageViewerOverlay::BindCommands()
 	);
 }
 
-TSharedRef<SWidget> SMediaImageViewerOverlay::CreateOverlay(bool bInShowPanelName)
+TSharedRef<SWidget> SMediaImageViewerOverlay::CreateOverlay(bool bInComparisonView)
 {
 	TSharedPtr<FMediaImageViewer> ImageViewer = Delegates->GetImageViewer.Execute(Position);
 
@@ -442,16 +435,57 @@ TSharedRef<SWidget> SMediaImageViewerOverlay::CreateOverlay(bool bInShowPanelNam
 
 	TSharedPtr<SWidget> ImageViewerOverlay = ImageViewer.IsValid() ? ImageViewer->GetOverlayWidget(Position) : nullptr;
 
-	Overlay->AddSlot()
-		[
-			SNew(SMediaViewerDropTarget, Delegates.ToSharedRef())
-			.Position(Position)
+	if (bInComparisonView)
+	{
+		Overlay->AddSlot()
 			[
-				ImageViewerOverlay.IsValid() ? ImageViewerOverlay.ToSharedRef() : SNullWidget::NullWidget
-			]
-		];
+				SNew(SMediaViewerDropTarget, Delegates.ToSharedRef())
+				.Position(Position)
+				.bComparisonView(bInComparisonView)
+			];
+	}
+	else
+	{
+		Overlay->AddSlot()
+			[
+				SNew(SHorizontalBox)
+				+ SHorizontalBox::Slot()
+				.FillWidth(0.25f)
+				.Padding(0.f, 0.f, 0.f, 20.f)
+				[
+					SNew(SMediaViewerDropTarget, Delegates.ToSharedRef())
+					.Position(EMediaImageViewerPosition::First)
+					.bComparisonView(bInComparisonView)
+					.bForceComparisonView(true)
+				]
+				+ SHorizontalBox::Slot()
+				.FillWidth(0.5f)
+				.Padding(10.f, 0.f, 10.f, 20.f)
+				[
+					SNew(SMediaViewerDropTarget, Delegates.ToSharedRef())
+					.Position(EMediaImageViewerPosition::First)
+					.bComparisonView(bInComparisonView)
+				]
+				+ SHorizontalBox::Slot()
+				.FillWidth(0.25f)
+				.Padding(0.f, 0.f, 0.f, 20.f)
+				[
+					SNew(SMediaViewerDropTarget, Delegates.ToSharedRef())
+					.Position(EMediaImageViewerPosition::Second)
+					.bComparisonView(bInComparisonView)
+				]
+			];
+	}
+
+	if (ImageViewerOverlay.IsValid())
+	{
+		Overlay->AddSlot()
+			[
+				ImageViewerOverlay.ToSharedRef()
+			];
+	}
 
-	if (bInShowPanelName)
+	if (bInComparisonView)
 	{
 		const EHorizontalAlignment HAlign = (Position == EMediaImageViewerPosition::First
 			&& Delegates->GetActiveView.Execute() == EMediaImageViewerActivePosition::Both
diff --git a/Engine/Plugins/Experimental/MediaViewer/Source/MediaViewer/Private/Widgets/SMediaImageViewerOverlay.h b/Engine/Plugins/Experimental/MediaViewer/Source/MediaViewer/Private/Widgets/SMediaImageViewerOverlay.h
--- a/Engine/Plugins/Experimental/MediaViewer/Source/MediaViewer/Private/Widgets/SMediaImageViewerOverlay.h
+++ b/Engine/Plugins/Experimental/MediaViewer/Source/MediaViewer/Private/Widgets/SMediaImageViewerOverlay.h
@@ -28,10 +28,10 @@ class SMediaImageViewerOverlay : public SCompoundWidget
 	SLATE_DECLARE_WIDGET(SMediaImageViewerOverlay, SCompoundWidget)
 
 	SLATE_BEGIN_ARGS(SMediaImageViewerOverlay)
-		: _bShowPanelName(true)
+		: _bComparisonView(true)
 		, _bScaleToFit(false)
 		{}
-		SLATE_ARGUMENT(bool, bShowPanelName)
+		SLATE_ARGUMENT(bool, bComparisonView)
 		SLATE_ARGUMENT(bool, bScaleToFit)
 	SLATE_END_ARGS()
 
@@ -99,7 +99,7 @@ protected:
 
 	void BindCommands();
 
-	TSharedRef<SWidget> CreateOverlay(bool bInShowPanelName);
+	TSharedRef<SWidget> CreateOverlay(bool bInComparisonView);
 
 	TSharedRef<SWidget> CreateStatusBar(const TSharedRef<FMediaViewerDelegates>& InDelegates);
 
diff --git a/Engine/Plugins/Experimental/MediaViewer/Source/MediaViewer/Private/Widgets/SMediaViewer.cpp b/Engine/Plugins/Experimental/MediaViewer/Source/MediaViewer/Private/Widgets/SMediaViewer.cpp
--- a/Engine/Plugins/Experimental/MediaViewer/Source/MediaViewer/Private/Widgets/SMediaViewer.cpp
+++ b/Engine/Plugins/Experimental/MediaViewer/Source/MediaViewer/Private/Widgets/SMediaViewer.cpp
@@ -18,7 +18,7 @@
 #include "Widgets/Layout/SSplitter.h"
 #include "Widgets/MediaViewerDelegates.h"
 #include "Widgets/SMediaImageViewerOverlay.h"
-#include "Widgets/SMediaViewerDropTarget.h"
+
 #include "Widgets/SMediaViewerLibrary.h"
 #include "Widgets/SMediaViewerLibraryPrivate.h"
 #include "Widgets/SMediaViewerToolbar.h"
@@ -223,11 +223,6 @@ bool SMediaViewer::IsShowingBothImageViewers() const
 
 void SMediaViewer::SwapABImageViewers()
 {
-	if (!IsShowingBothImageViewers())
-	{
-		return;
-	}
-
 	TSharedPtr<FMediaImageViewer> OldSecond = ImageViewers[static_cast<int32>(EMediaImageViewerPosition::Second)];	
 	ImageViewers[static_cast<int32>(EMediaImageViewerPosition::Second)] = ImageViewers[static_cast<int32>(EMediaImageViewerPosition::First)];
 	ImageViewers[static_cast<int32>(EMediaImageViewerPosition::First)] = OldSecond;
@@ -251,6 +246,18 @@ void SMediaViewer::SetSingleView()
 	InvalidateView();
 }
 
+void SMediaViewer::SetABView()
+{
+	if (ActiveView == EMediaImageViewerActivePosition::Both)
+	{
+		return;
+	}
+
+	RequestedView = EMediaImageViewerActivePosition::Both;
+
+	InvalidateView();
+}
+
 void SMediaViewer::SetABOrientation(EOrientation InOrientation)
 {
 	if (MediaViewerSettings.ABOrientation == InOrientation
@@ -487,6 +494,7 @@ void SMediaViewer::CreateDelegates()
 
 	// Global viewer delegates
 	Delegates->SetSingleView.BindSP(this, &SMediaViewer::SetSingleView);
+	Delegates->SetABView.BindSP(this, &SMediaViewer::SetABView);
 	Delegates->GetABOrientation.BindSPLambda(this, [this] { return MediaViewerSettings.ABOrientation; });
 	Delegates->SetABOrientation.BindSP(this, &SMediaViewer::SetABOrientation);
 	Delegates->GetActiveView.BindSPLambda(this, [this] { return ActiveView; });
@@ -521,7 +529,7 @@ void SMediaViewer::CreateDelegates()
 	Delegates->GetCommandListForPosition.BindSP(this, &SMediaViewer::GetOverlayCommandList);
 }
 
-TSharedRef<SMediaImageViewerOverlay> SMediaViewer::CreateOverlay(EMediaImageViewerPosition InPosition, bool bInShowPanelName)
+TSharedRef<SMediaImageViewerOverlay> SMediaViewer::CreateOverlay(EMediaImageViewerPosition InPosition, bool bInComparisonView)
 {
 	bool bScaleToFitImage = false;
 	const int32 Index = static_cast<int32>(InPosition);
@@ -533,7 +541,7 @@ TSharedRef<SMediaImageViewerOverlay> SMediaViewer::CreateOverlay(EMediaImageView
 	}
 
 	return SNew(SMediaImageViewerOverlay, InPosition, Delegates.ToSharedRef())
-		.bShowPanelName(bInShowPanelName)
+		.bComparisonView(bInComparisonView)
 		.bScaleToFit(bScaleToFitImage);
 }
 
diff --git a/Engine/Plugins/Experimental/MediaViewer/Source/MediaViewer/Private/Widgets/SMediaViewer.h b/Engine/Plugins/Experimental/MediaViewer/Source/MediaViewer/Private/Widgets/SMediaViewer.h
--- a/Engine/Plugins/Experimental/MediaViewer/Source/MediaViewer/Private/Widgets/SMediaViewer.h
+++ b/Engine/Plugins/Experimental/MediaViewer/Source/MediaViewer/Private/Widgets/SMediaViewer.h
@@ -72,9 +72,10 @@ public:
 
 	void SwapABImageViewers();
 
-	/** Removes the second/B image and goes into single view. */
 	void SetSingleView();
 
+	void SetABView();
+
 	/** Sets the orientation of the A/B view. */
 	void SetABOrientation(EOrientation InOrientation);
 
@@ -147,7 +148,7 @@ protected:
 	void BindCommands();
 
 	/** Create and overlay widget for the given position. */
-	TSharedRef<SMediaImageViewerOverlay> CreateOverlay(EMediaImageViewerPosition InPosition, bool bInShowPanelName);
+	TSharedRef<SMediaImageViewerOverlay> CreateOverlay(EMediaImageViewerPosition InPosition, bool bInComparisonView);
 
 	/** Creates the toolbar. */
 	TSharedRef<SWidget> CreateToolbar();
diff --git a/Engine/Plugins/Experimental/MediaViewer/Source/MediaViewer/Private/Widgets/SMediaViewerDropTarget.cpp b/Engine/Plugins/Experimental/MediaViewer/Source/MediaViewer/Private/Widgets/SMediaViewerDropTarget.cpp
--- a/Engine/Plugins/Experimental/MediaViewer/Source/MediaViewer/Private/Widgets/SMediaViewerDropTarget.cpp
+++ b/Engine/Plugins/Experimental/MediaViewer/Source/MediaViewer/Private/Widgets/SMediaViewerDropTarget.cpp
@@ -12,6 +12,7 @@
 #include "Library/MediaViewerLibrary.h"
 #include "Library/MediaViewerLibraryItem.h"
 #include "MediaPlayer.h"
+#include "SDropTarget.h"
 #include "Styling/AppStyle.h"
 #include "Styling/StyleColors.h"
 #include "Widgets/Layout/SBox.h"
@@ -33,31 +34,66 @@ void SMediaViewerDropTarget::Construct(const FArguments& InArgs, const TSharedRe
 {
 	Delegates = InDelegates;
 	Position = InArgs._Position;
+	bComparisonView = InArgs._bComparisonView;
+	bForceComparisonView = InArgs._bForceComparisonView;
 
-	SDropTarget::Construct(
-		SDropTarget::FArguments()
-		.OnDropped(this, &SMediaViewerDropTarget::OnDropped)
+	FText UpperMessage;
+
+	if (bComparisonView || (Position == EMediaImageViewerPosition::First && !bForceComparisonView))
+	{
+		UpperMessage = LOCTEXT("ReplaceImage", "Replace Image");
+	}
+	else
+	{
+		UpperMessage = LOCTEXT("CompareImage", "Compare Image");
+	}
+
+	ChildSlot
+	[
+		SNew(SOverlay)
+		+ SOverlay::Slot()
 		[
-			SNew(SOverlay)
-			+ SOverlay::Slot()
+			SNew(SDropTarget)
+			.OnAllowDrop(this, &SMediaViewerDropTarget::OnAllowDrop)
+			.OnIsRecognized(this, &SMediaViewerDropTarget::OnIsRecognized)
+			.OnDropped(this, &SMediaViewerDropTarget::OnDropped)
+		]
+		+ SOverlay::Slot()
+		[
+			InArgs._Content.Widget
+		]
+		+ SOverlay::Slot()
+		[
+			SNew(SVerticalBox)
+			+ SVerticalBox::Slot()
+			.FillHeight(1.f)
+			.HAlign(HAlign_Center)
+			.VAlign(VAlign_Bottom)
+			.Padding(0.f, 0.f, 0.f, 5.f)
 			[
-				InArgs._Content.Widget
+				SNew(STextBlock)
+				.ColorAndOpacity(FStyleColors::Foreground.GetSpecifiedColor())
+				.ShadowColorAndOpacity(FStyleColors::Panel.GetSpecifiedColor())
+				.ShadowOffset(FVector2D(1.f, 1.f))
+				.Visibility(this, &SMediaViewerDropTarget::GetDragDescriptionVisibility)
+				.Text(UpperMessage)
 			]
-			+ SOverlay::Slot()
+			+ SVerticalBox::Slot()
+			.FillHeight(1.f)
+			.HAlign(HAlign_Center)
+			.VAlign(VAlign_Top)
+			.Padding(5.f)
 			[
-				SNew(SBox)
-				.HAlign(HAlign_Center)
-				.VAlign(VAlign_Center)
-				[
-					SNew(STextBlock)
-					.ColorAndOpacity(FStyleColors::Foreground.GetSpecifiedColor())
-					.ShadowColorAndOpacity(FStyleColors::Panel.GetSpecifiedColor())
-					.ShadowOffset(FVector2D(1.f, 1.f))
-					.Visibility(this, &SMediaViewerDropTarget::GetDragDescriptionVisibility)
-					.Text(LOCTEXT("DropTargetMessage", "Drop supported asset or library item here."))
-				]
+				SNew(STextBlock)
+				.ColorAndOpacity(FStyleColors::Foreground.GetSpecifiedColor())
+				.ShadowColorAndOpacity(FStyleColors::Panel.GetSpecifiedColor())
+				.ShadowOffset(FVector2D(1.f, 1.f))
+				.Visibility(this, &SMediaViewerDropTarget::GetDragDescriptionVisibility)
+				.Text(LOCTEXT("DropTargetMessage", "Drop supported asset or library item here."))
+				.AutoWrapText(true)
 			]
-		]);
+		]
+	];
 }
 
 void SMediaViewerDropTarget::OnDragLeave(const FDragDropEvent& InDragDropEvent)
@@ -134,6 +170,13 @@ void SMediaViewerDropTarget::HandleDroppedMediaViewerOp(const FMediaViewerLibrar
 	{
 		if (TSharedPtr<FMediaImageViewer> ImageViewer = LibraryItem->CreateImageViewer())
 		{
+			// Swap images and then replace first image.
+			if (bForceComparisonView && Position == EMediaImageViewerPosition::First)
+			{
+				Delegates->SwapAB.Execute();
+				Delegates->SetABView.Execute();
+			}
+
 			Delegates->SetImageViewer.Execute(Position, ImageViewer.ToSharedRef());
 		}
 	}
@@ -171,6 +214,13 @@ void SMediaViewerDropTarget::HandleDroppedAssets(const TArrayView<FAssetData> In
 	}
 	else if (ImageViewers.Num() == 1)
 	{
+		// Swap images and then replace first image.
+		if (bForceComparisonView && Position == EMediaImageViewerPosition::First)
+		{
+			Delegates->SwapAB.Execute();
+			Delegates->SetABView.Execute();
+		}
+
 		Delegates->SetImageViewer.Execute(Position, ImageViewers[0]);
 	}
 }
@@ -200,11 +250,28 @@ void SMediaViewerDropTarget::HandleDroppeFileOp(const FExternalDragOperation& In
 
 	TestMediaPlayer->Close();
 
-	for (int32 Index = 0; Index < Sources.Num(); ++Index)
+	if (Sources.Num() > 1)
+	{
+		for (int32 Index = 0; Index < Sources.Num(); ++Index)
+		{
+			Delegates->SetImageViewer.Execute(
+				static_cast<EMediaImageViewerPosition>(Index),
+				MakeShared<FMediaSourceImageViewer>(Sources[Index], FText::FromString(Sources[Index]->GetFilePath()))
+			);
+		}
+	}
+	else if (Sources.Num() == 1)
 	{
+		// Swap images and then replace first image.
+		if (bForceComparisonView && Position == EMediaImageViewerPosition::First)
+		{
+			Delegates->SwapAB.Execute();
+			Delegates->SetABView.Execute();
+		}
+
 		Delegates->SetImageViewer.Execute(
-			static_cast<EMediaImageViewerPosition>(Index),
-			MakeShared<FMediaSourceImageViewer>(Sources[Index], FText::FromString(Sources[Index]->GetFilePath()))
+			EMediaImageViewerPosition::First,
+			MakeShared<FMediaSourceImageViewer>(Sources[0], FText::FromString(Sources[0]->GetFilePath()))
 		);
 	}
 }
@@ -222,7 +289,7 @@ bool SMediaViewerDropTarget::OnAllowDrop(TSharedPtr<FDragDropOperation> InDragDr
 
 		if (!Delegates->GetLibrary.Execute()->GetItem(LibraryItemDragDropOp->GetGroupItem().ItemId).IsValid())
 		{
-			if (IsDragOver())
+			if (IsHovered())
 			{
 				LibraryItemDragDropOp->SetToolTip(LOCTEXT("InvalidItem", "Invalid Library Item"), FAppStyle::GetBrush(TEXT("Graph.ConnectorFeedback.Error")));
 			}
diff --git a/Engine/Plugins/Experimental/MediaViewer/Source/MediaViewer/Private/Widgets/SMediaViewerDropTarget.h b/Engine/Plugins/Experimental/MediaViewer/Source/MediaViewer/Private/Widgets/SMediaViewerDropTarget.h
--- a/Engine/Plugins/Experimental/MediaViewer/Source/MediaViewer/Private/Widgets/SMediaViewerDropTarget.h
+++ b/Engine/Plugins/Experimental/MediaViewer/Source/MediaViewer/Private/Widgets/SMediaViewerDropTarget.h
@@ -2,7 +2,7 @@
 
 #pragma once
 
-#include "SDropTarget.h"
+#include "Widgets/SCompoundWidget.h"
 
 #include "SMediaViewer.h"
 
@@ -12,15 +12,19 @@ namespace UE::MediaViewer::Private
 class FMediaViewerLibraryItemDragDropOperation;
 struct FMediaViewerDelegates;
 
-class SMediaViewerDropTarget : public SDropTarget
+class SMediaViewerDropTarget : public SCompoundWidget
 {
-	SLATE_DECLARE_WIDGET(SMediaViewerDropTarget, SDropTarget)
+	SLATE_DECLARE_WIDGET(SMediaViewerDropTarget, SCompoundWidget)
 
 	SLATE_BEGIN_ARGS(SMediaViewerDropTarget)
 		: _Position(EMediaImageViewerPosition::First)
+		, _bComparisonView(false)
+		, _bForceComparisonView(false)
 		{}
 		SLATE_DEFAULT_SLOT(FArguments, Content)
 		SLATE_ARGUMENT(EMediaImageViewerPosition, Position)
+		SLATE_ARGUMENT(bool, bComparisonView)
+		SLATE_ARGUMENT(bool, bForceComparisonView)
 	SLATE_END_ARGS()
 
 public:
@@ -35,9 +39,15 @@ protected:
 
 	TSharedPtr<FMediaViewerDelegates> Delegates;
 	EMediaImageViewerPosition Position;
+	bool bComparisonView;
+	bool bForceComparisonView;
 
 	EVisibility GetDragDescriptionVisibility() const;
 
+	bool OnAllowDrop(TSharedPtr<FDragDropOperation> InDragDropOperation) const;
+
+	bool OnIsRecognized(TSharedPtr<FDragDropOperation> InDragDropOperation) const;
+
 	FReply OnDropped(const FGeometry& InGeometry, const FDragDropEvent& InDragDropEvent);
 
 	void HandleDroppedMediaViewerOp(const FMediaViewerLibraryItemDragDropOperation& InMediaViewerOp);
@@ -45,11 +55,6 @@ protected:
 	void HandleDroppedAssets(const TArrayView<FAssetData> InDroppedAssets);
 	
 	void HandleDroppeFileOp(const FExternalDragOperation& InFileDragDropOp);
-
-	//~ Begin SDropTarget
-	virtual bool OnAllowDrop(TSharedPtr<FDragDropOperation> InDragDropOperation) const override;
-	virtual bool OnIsRecognized(TSharedPtr<FDragDropOperation> InDragDropOperation) const override;
-	//~ End SDropTarget
 };
 
 } // UE::MediaViewer::Private
