Media Viewer: Drop target update

diff --git a/Engine/Plugins/Experimental/MediaViewer/Source/MediaViewer/Private/Widgets/SMediaViewerDropTarget.cpp b/Engine/Plugins/Experimental/MediaViewer/Source/MediaViewer/Private/Widgets/SMediaViewerDropTarget.cpp
--- a/Engine/Plugins/Experimental/MediaViewer/Source/MediaViewer/Private/Widgets/SMediaViewerDropTarget.cpp
+++ b/Engine/Plugins/Experimental/MediaViewer/Source/MediaViewer/Private/Widgets/SMediaViewerDropTarget.cpp
@@ -175,6 +175,7 @@ void SMediaViewerDropTarget::HandleDroppedMediaViewerOp(const FMediaViewerLibrar
 			{
 				Delegates->SwapAB.Execute();
 				Delegates->SetABView.Execute();
+				Delegates->SetABOrientation.Execute(Orient_Horizontal);
 			}
 
 			Delegates->SetImageViewer.Execute(Position, ImageViewer.ToSharedRef());
@@ -219,6 +220,7 @@ void SMediaViewerDropTarget::HandleDroppedAssets(const TArrayView<FAssetData> In
 		{
 			Delegates->SwapAB.Execute();
 			Delegates->SetABView.Execute();
+			Delegates->SetABOrientation.Execute(Orient_Horizontal);
 		}
 
 		Delegates->SetImageViewer.Execute(Position, ImageViewers[0]);
@@ -267,6 +269,7 @@ void SMediaViewerDropTarget::HandleDroppeFileOp(const FExternalDragOperation& In
 		{
 			Delegates->SwapAB.Execute();
 			Delegates->SetABView.Execute();
+			Delegates->SetABOrientation.Execute(Orient_Horizontal);
 		}
 
 		Delegates->SetImageViewer.Execute(
