Media Viewer: Media playback buttons changed to be more like sequencer

diff --git a/Engine/Plugins/Experimental/MediaViewer/Source/MediaViewer/Private/Widgets/SMediaSourceOverlay.cpp b/Engine/Plugins/Experimental/MediaViewer/Source/MediaViewer/Private/Widgets/SMediaSourceOverlay.cpp
--- a/Engine/Plugins/Experimental/MediaViewer/Source/MediaViewer/Private/Widgets/SMediaSourceOverlay.cpp
+++ b/Engine/Plugins/Experimental/MediaViewer/Source/MediaViewer/Private/Widgets/SMediaSourceOverlay.cpp
@@ -54,7 +54,7 @@ void SMediaSourceOverlay::Construct(const FArguments& InArgs, const TSharedRef<F
 			+ SVerticalBox::Slot()
 			.AutoHeight()
 			.HAlign(HAlign_Fill)
-			.Padding(0.f, 0.f, 0.f, 0.f)
+			.Padding(0.f, 2.f, 0.f, 0.f)
 			[
 				CreateControls()
 			]
@@ -233,8 +233,8 @@ TSharedRef<SWidget> SMediaSourceOverlay::CreateControls()
 			[
 				SNew(SImage)
 				.ColorAndOpacity(FSlateColor::UseForeground())
-				.Image(FSlateIcon(FAppStyle::GetAppStyleSetName(), "Animation.Backward").GetIcon())
-				.ToolTipText(LOCTEXT("Reverse", "Play media in reverse.\n\nNot widely supported by media decoders."))
+				.Image(this, &SMediaSourceOverlay::Reverse_GetBrush)
+				.ToolTipText(this, &SMediaSourceOverlay::Reverse_GetToolTip)
 				.DesiredSizeOverride(FVector2D(20.f))
 			]
 		]
@@ -247,36 +247,14 @@ TSharedRef<SWidget> SMediaSourceOverlay::CreateControls()
 			SNew(SButton)
 			.VAlign(VAlign_Center)
 			.IsEnabled(this, &SMediaSourceOverlay::Play_IsEnabled)
-			.Visibility(this, &SMediaSourceOverlay::Play_GetVisibility)
 			.OnClicked(this, &SMediaSourceOverlay::Play_OnClicked)
 			.ButtonStyle(FMediaViewerStyle::Get(), "MediaButtons")
 			.ContentPadding(2.f)
 			[
 				SNew(SImage)
 				.ColorAndOpacity(FSlateColor::UseForeground())
-				.Image(FSlateIcon(FAppStyle::GetAppStyleSetName(), "Animation.Forward").GetIcon())
-				.ToolTipText(LOCTEXT("Play", "Start media playback"))
-				.DesiredSizeOverride(FVector2D(20.f))
-			]
-		]
-
-		// Pause button.
-		+ SHorizontalBox::Slot()
-		.VAlign(VAlign_Center)
-		.AutoWidth()
-		[
-			SNew(SButton)
-			.VAlign(VAlign_Center)
-			.IsEnabled(this, &SMediaSourceOverlay::Pause_IsEnabled)
-			.Visibility(this, &SMediaSourceOverlay::Pause_GetVisibility)
-			.OnClicked(this, &SMediaSourceOverlay::Pause_OnClicked)
-			.ButtonStyle(FMediaViewerStyle::Get(), "MediaButtons")
-			.ContentPadding(2.f)
-			[
-				SNew(SImage)
-				.ColorAndOpacity(FSlateColor::UseForeground())
-				.Image(FSlateIcon(FAppStyle::GetAppStyleSetName(), "Animation.Pause").GetIcon())
-				.ToolTipText(LOCTEXT("Pause", "Pause media playback"))
+				.Image(this, &SMediaSourceOverlay::Play_GetBrush)
+				.ToolTipText(this, &SMediaSourceOverlay::Play_GetToolTip)
 				.DesiredSizeOverride(FVector2D(20.f))
 			]
 		]
@@ -479,115 +457,142 @@ bool SMediaSourceOverlay::Reverse_IsEnabled() const
 {
 	if (UMediaPlayer* MediaPlayer = GetMediaPlayer())
 	{
-		const float Rate = FMath::Abs(MediaPlayer->GetRate());
-
-		return MediaPlayer->IsReady() && MediaPlayer->SupportsRate(Rate > 0 ? -Rate : -1.f, /* Unthinned */ true);
+		return MediaPlayer->IsReady() && MediaPlayer->SupportsRate(-1.f, /* Unthinned */ true);
 	}
 
 	return false;
 }
 
-FReply SMediaSourceOverlay::Reverse_OnClicked()
+const FSlateBrush* SMediaSourceOverlay::Reverse_GetBrush() const
 {
 	if (UMediaPlayer* MediaPlayer = GetMediaPlayer())
 	{
 		const float Rate = FMath::Abs(MediaPlayer->GetRate());
 
-		MediaPlayer->SetRate(Rate > 0 ? -Rate : -1.f);
+		if (Rate < 0 && !FMath::IsNearlyZero(Rate))
+		{
+			return FAppStyle::Get().GetBrush("Animation.Pause");
+		}
 	}
 
-	return FReply::Handled();
+	return FAppStyle::Get().GetBrush("Animation.Backward");
 }
 
-bool SMediaSourceOverlay::StepBack_IsEnabled() const
+FText SMediaSourceOverlay::Reverse_GetToolTip() const
 {
 	if (UMediaPlayer* MediaPlayer = GetMediaPlayer())
 	{
 		const float Rate = FMath::Abs(MediaPlayer->GetRate());
 
-		return MediaPlayer->IsReady() && MediaPlayer->IsPaused();
+		if (Rate < 0 && !FMath::IsNearlyZero(Rate))
+		{
+			return LOCTEXT("Pause", "Pause media playback");
+		}
 	}
 
-	return false;
+	return LOCTEXT("Reverse", "Play media in reverse.\n\nNot widely supported by media decoders.");
 }
 
-FReply SMediaSourceOverlay::StepBack_OnClicked()
+FReply SMediaSourceOverlay::Reverse_OnClicked()
 {
 	if (UMediaPlayer* MediaPlayer = GetMediaPlayer())
 	{
-		const float FrameRate = MediaPlayer->GetVideoTrackFrameRate(INDEX_NONE, INDEX_NONE);
+		const float Rate = FMath::Abs(MediaPlayer->GetRate());
 
-		if (!FMath::IsNearlyZero(FrameRate))
+		if (FMath::IsNearlyZero(Rate))
 		{
-			MediaPlayer->Seek(MediaPlayer->GetTime() - FTimespan::FromSeconds(1.f / FrameRate));
+			MediaPlayer->SetRate(-1.f);
+		}
+		else
+		{
+			MediaPlayer->Pause();
 		}
 	}
 
 	return FReply::Handled();
 }
 
-bool SMediaSourceOverlay::Play_IsEnabled() const
+bool SMediaSourceOverlay::StepBack_IsEnabled() const
 {
 	if (UMediaPlayer* MediaPlayer = GetMediaPlayer())
 	{
-		return MediaPlayer->IsReady()
-			&& (!MediaPlayer->IsPlaying() || (MediaPlayer->GetRate() != 1.0f));
+		const float Rate = FMath::Abs(MediaPlayer->GetRate());
+
+		return MediaPlayer->IsReady() && MediaPlayer->IsPaused();
 	}
 
 	return false;
 }
 
-EVisibility SMediaSourceOverlay::Play_GetVisibility() const
+FReply SMediaSourceOverlay::StepBack_OnClicked()
 {
 	if (UMediaPlayer* MediaPlayer = GetMediaPlayer())
 	{
-		return MediaPlayer->IsPlaying() 
-			? EVisibility::Collapsed 
-			: EVisibility::Visible;
+		const float FrameRate = MediaPlayer->GetVideoTrackFrameRate(INDEX_NONE, INDEX_NONE);
+
+		if (!FMath::IsNearlyZero(FrameRate))
+		{
+			MediaPlayer->Seek(MediaPlayer->GetTime() - FTimespan::FromSeconds(1.f / FrameRate));
+		}
 	}
 
-	return EVisibility::Visible;
+	return FReply::Handled();
 }
 
-FReply SMediaSourceOverlay::Play_OnClicked()
+bool SMediaSourceOverlay::Play_IsEnabled() const
 {
 	if (UMediaPlayer* MediaPlayer = GetMediaPlayer())
 	{
-		MediaPlayer->Play();
+		return MediaPlayer->IsReady();
 	}
 
-	return FReply::Handled();
+	return false;
 }
 
-bool SMediaSourceOverlay::Pause_IsEnabled() const
+const FSlateBrush* SMediaSourceOverlay::Play_GetBrush() const
 {
 	if (UMediaPlayer* MediaPlayer = GetMediaPlayer())
 	{
-		return MediaPlayer->CanPause() && !MediaPlayer->IsPaused();
+		const float Rate = FMath::Abs(MediaPlayer->GetRate());
+
+		if (Rate > 0 && !FMath::IsNearlyZero(Rate))
+		{
+			return FAppStyle::Get().GetBrush("Animation.Pause");
+		}
 	}
 
-	return false;
+	return FAppStyle::Get().GetBrush("Animation.Forward");
 }
 
-EVisibility SMediaSourceOverlay::Pause_GetVisibility() const
+FText SMediaSourceOverlay::Play_GetToolTip() const
 {
 	if (UMediaPlayer* MediaPlayer = GetMediaPlayer())
 	{
-		const bool bIsVisible = MediaPlayer->CanPause() && !MediaPlayer->IsPaused();
-		
-		return bIsVisible 
-			? EVisibility::Visible 
-			: EVisibility::Collapsed;
+		const float Rate = FMath::Abs(MediaPlayer->GetRate());
+
+		if (Rate > 0 && !FMath::IsNearlyZero(Rate))
+		{
+			return LOCTEXT("Pause", "Pause media playback");
+		}
 	}
 
-	return EVisibility::Visible;
+	return LOCTEXT("Play", "Play media forward");
 }
 
-FReply SMediaSourceOverlay::Pause_OnClicked()
+FReply SMediaSourceOverlay::Play_OnClicked()
 {
 	if (UMediaPlayer* MediaPlayer = GetMediaPlayer())
 	{
-		MediaPlayer->Pause();
+		const float Rate = FMath::Abs(MediaPlayer->GetRate());
+
+		if (FMath::IsNearlyZero(Rate))
+		{
+			MediaPlayer->Play();
+		}
+		else
+		{
+			MediaPlayer->Pause();
+		}
 	}
 
 	return FReply::Handled();
diff --git a/Engine/Plugins/Experimental/MediaViewer/Source/MediaViewer/Private/Widgets/SMediaSourceOverlay.h b/Engine/Plugins/Experimental/MediaViewer/Source/MediaViewer/Private/Widgets/SMediaSourceOverlay.h
--- a/Engine/Plugins/Experimental/MediaViewer/Source/MediaViewer/Private/Widgets/SMediaSourceOverlay.h
+++ b/Engine/Plugins/Experimental/MediaViewer/Source/MediaViewer/Private/Widgets/SMediaSourceOverlay.h
@@ -10,6 +10,7 @@
 class IMediaStreamPlayer;
 class UMediaPlayer;
 class UMediaStream;
+struct FSlateBrush;
 
 namespace UE::MediaViewer::Private
 {
@@ -57,19 +58,18 @@ protected:
 	FReply Rewind_OnClicked();
 
 	bool Reverse_IsEnabled() const;
+	const FSlateBrush* Reverse_GetBrush() const;
+	FText Reverse_GetToolTip() const;
 	FReply Reverse_OnClicked();
 
 	bool StepBack_IsEnabled() const;
 	FReply StepBack_OnClicked();
 
 	bool Play_IsEnabled() const;
-	EVisibility Play_GetVisibility() const;
+	const FSlateBrush* Play_GetBrush() const;
+	FText Play_GetToolTip() const;
 	FReply Play_OnClicked();
 
-	bool Pause_IsEnabled() const;
-	EVisibility Pause_GetVisibility() const;
-	FReply Pause_OnClicked();
-
 	bool StepForward_IsEnabled() const;
 	FReply StepForward_OnClicked();
 
