Media Viewer: Right clicking with no image viewer open no longer causes a crash.

diff --git a/Engine/Plugins/Experimental/MediaViewer/Source/MediaViewer/Private/Widgets/SMediaImageViewerOverlay.cpp b/Engine/Plugins/Experimental/MediaViewer/Source/MediaViewer/Private/Widgets/SMediaImageViewerOverlay.cpp
--- a/Engine/Plugins/Experimental/MediaViewer/Source/MediaViewer/Private/Widgets/SMediaImageViewerOverlay.cpp
+++ b/Engine/Plugins/Experimental/MediaViewer/Source/MediaViewer/Private/Widgets/SMediaImageViewerOverlay.cpp
@@ -857,6 +857,11 @@ void SMediaImageViewerOverlay::CopyColor()
 
 bool SMediaImageViewerOverlay::CanAddToLibrary() const
 {
+	if (!CachedItem.IsValid())
+	{
+		return false;
+	}
+
 	TSharedRef<FMediaViewerLibrary> Library = Delegates->GetLibrary.Execute();
 
 	return !Library->FindItemByValue(CachedItem->GetItemType(), CachedItem->GetStringValue()).IsValid();
@@ -864,6 +869,11 @@ bool SMediaImageViewerOverlay::CanAddToLibrary() const
 
 void SMediaImageViewerOverlay::AddToLibrary()
 {
+	if (!CachedItem.IsValid())
+	{
+		return;
+	}
+
 	TSharedRef<FMediaViewerLibrary> Library = Delegates->GetLibrary.Execute();
 
 	Library->AddItemToGroup(CachedItem.ToSharedRef());
