Media Viewer: Added player name to the status bar for media viewers.

diff --git a/Engine/Plugins/Experimental/MediaViewer/Source/MediaViewer/Private/ImageViewers/MediaSourceImageViewer.cpp b/Engine/Plugins/Experimental/MediaViewer/Source/MediaViewer/Private/ImageViewers/MediaSourceImageViewer.cpp
--- a/Engine/Plugins/Experimental/MediaViewer/Source/MediaViewer/Private/ImageViewers/MediaSourceImageViewer.cpp
+++ b/Engine/Plugins/Experimental/MediaViewer/Source/MediaViewer/Private/ImageViewers/MediaSourceImageViewer.cpp
@@ -4,6 +4,7 @@
 
 #include "AssetRegistry/AssetData.h"
 #include "Brushes/SlateImageBrush.h"
+#include "DetailLayoutBuilder.h"
 #include "Engine/Texture2D.h"
 #include "FileMediaSource.h"
 #include "IMediaStreamPlayer.h"
@@ -19,7 +20,10 @@
 #include "Misc/Optional.h"
 #include "Misc/TVariant.h"
 #include "Widgets/Layout/SBox.h"
+#include "Widgets/MediaImageStatusBarExtender.h"
+#include "Widgets/SBoxPanel.h"
 #include "Widgets/SMediaSourceOverlay.h"
+#include "Widgets/Text/STextBlock.h"
 
 #define LOCTEXT_NAMESPACE "MediaSourceImageViewer"
 
@@ -421,6 +425,55 @@ TSharedPtr<SWidget> FMediaSourceImageViewer::GetOverlayWidget(EMediaImageViewerP
 		];
 }
 
+void FMediaSourceImageViewer::ExtendStatusBar(UE::MediaViewer::FMediaImageStatusBarExtender& InOutStatusBarExtender)
+{
+	InOutStatusBarExtender.AddExtension(
+		UE::MediaViewer::StatusBarSections::StatusBarCenter,
+		EExtensionHook::Before,
+		nullptr,
+		UE::MediaViewer::FMediaImageStatusBarExtension::FDelegate::CreateSP(this, &FMediaSourceImageViewer::AddPlayerName)
+	);
+}
+
+void FMediaSourceImageViewer::AddPlayerName(const TSharedRef<SHorizontalBox>& InStatusBar)
+{
+	if (!MediaStream)
+	{
+		return;
+	}
+
+	IMediaStreamPlayer* MediaStreamPlayer = MediaStream->GetPlayer().GetInterface();
+
+	if (!MediaStreamPlayer)
+	{
+		return;
+	}
+
+	UMediaPlayer* MediaPlayer = MediaStreamPlayer->GetPlayer();
+
+	if (!MediaPlayer)
+	{
+		return;
+	}
+
+	const FName PlayerName = MediaPlayer->GetPlayerName();
+
+	if (PlayerName.IsNone())
+	{
+		return;
+	}
+
+	InStatusBar->AddSlot()
+		.AutoWidth()
+		.Padding(2.f)
+		.VAlign(VAlign_Center)
+		[
+			SNew(STextBlock)
+			.Text(FText::FromName(PlayerName))
+			.Font(IDetailLayoutBuilder::GetDetailFont())
+		];
+}
+
 FString FMediaSourceImageViewer::GetReferencerName() const
 {
 	static const FString ReferencerName = TEXT("FMediaSourceImageViewer");
diff --git a/Engine/Plugins/Experimental/MediaViewer/Source/MediaViewer/Private/ImageViewers/MediaSourceImageViewer.h b/Engine/Plugins/Experimental/MediaViewer/Source/MediaViewer/Private/ImageViewers/MediaSourceImageViewer.h
--- a/Engine/Plugins/Experimental/MediaViewer/Source/MediaViewer/Private/ImageViewers/MediaSourceImageViewer.h
+++ b/Engine/Plugins/Experimental/MediaViewer/Source/MediaViewer/Private/ImageViewers/MediaSourceImageViewer.h
@@ -15,6 +15,7 @@
 
 #include "MediaSourceImageViewer.generated.h"
 
+class SHorizontalBox;
 class UMediaSource;
 class UMediaStream;
 
@@ -121,6 +122,7 @@ public:
 	virtual TOptional<TVariant<FColor, FLinearColor>> GetPixelColor(const FIntPoint& InPixelCoords, int32 InMipLevel) const override;
 	virtual TSharedPtr<FStructOnScope> GetCustomSettingsOnScope() const override;
 	virtual TSharedPtr<SWidget> GetOverlayWidget(EMediaImageViewerPosition InPosition) override;
+	virtual void ExtendStatusBar(UE::MediaViewer::FMediaImageStatusBarExtender& InOutStatusBarExtender) override;
 	//~ End FMediaImageViewer
 
 	//~ Begin FGCObject
@@ -133,6 +135,8 @@ protected:
 	TObjectPtr<UMediaStream> MediaStream;
 	mutable FTextureSampleCache SampleCache;
 
+	void AddPlayerName(const TSharedRef<SHorizontalBox>& InStatusBar);
+
 	//~ Begin FMediaImageViewer
 	virtual void PaintImage(FMediaImagePaintParams& InPaintParams, const FMediaImagePaintGeometry& InPaintGeometry) override;
 	//~ End FMediaImageViewer
diff --git a/Engine/Plugins/Experimental/MediaViewer/Source/MediaViewer/Private/Widgets/MediaImageStatusBarExtender.cpp b/Engine/Plugins/Experimental/MediaViewer/Source/MediaViewer/Private/Widgets/MediaImageStatusBarExtender.cpp
--- a/Engine/Plugins/Experimental/MediaViewer/Source/MediaViewer/Private/Widgets/MediaImageStatusBarExtender.cpp
+++ b/Engine/Plugins/Experimental/MediaViewer/Source/MediaViewer/Private/Widgets/MediaImageStatusBarExtender.cpp
@@ -16,7 +16,7 @@ void FMediaImageStatusBarExtender::AddExtension(FName InExtensionHook, EExtensio
 }
 
 void FMediaImageStatusBarExtender::Apply(FName InExtensionHook, EExtensionHook::Position InHookPosition,
-	SHorizontalBox& InHorizontalBox) const
+	const TSharedRef<SHorizontalBox>& InHorizontalBox) const
 {
 	for (const FMediaImageStatusBarExtension& Extension : Extensions)
 	{
diff --git a/Engine/Plugins/Experimental/MediaViewer/Source/MediaViewer/Private/Widgets/SMediaImageViewerStatusBar.cpp b/Engine/Plugins/Experimental/MediaViewer/Source/MediaViewer/Private/Widgets/SMediaImageViewerStatusBar.cpp
--- a/Engine/Plugins/Experimental/MediaViewer/Source/MediaViewer/Private/Widgets/SMediaImageViewerStatusBar.cpp
+++ b/Engine/Plugins/Experimental/MediaViewer/Source/MediaViewer/Private/Widgets/SMediaImageViewerStatusBar.cpp
@@ -59,7 +59,7 @@ TSharedRef<SWidget> SMediaImageViewerStatusBar::BuildStatusBar()
 	{
 		if (ExtensionHook != NAME_None)
 		{
-			StatusBarExtender.Apply(ExtensionHook, HookPosition, *HorizontalBox);
+			StatusBarExtender.Apply(ExtensionHook, HookPosition, HorizontalBox);
 		}
 	};
 
diff --git a/Engine/Plugins/Experimental/MediaViewer/Source/MediaViewer/Public/Widgets/MediaImageStatusBarExtender.h b/Engine/Plugins/Experimental/MediaViewer/Source/MediaViewer/Public/Widgets/MediaImageStatusBarExtender.h
--- a/Engine/Plugins/Experimental/MediaViewer/Source/MediaViewer/Public/Widgets/MediaImageStatusBarExtender.h
+++ b/Engine/Plugins/Experimental/MediaViewer/Source/MediaViewer/Public/Widgets/MediaImageStatusBarExtender.h
@@ -22,7 +22,7 @@ struct FMediaImageStatusBarExtension
 	 * The given @see SHorizontalBox is the layout container that the widget for the extension needs to be added to by calling.
 	 * @see SHorizontalBox::AddSlot().
 	 */
-	DECLARE_DELEGATE_OneParam(FDelegate, SHorizontalBox&);
+	DECLARE_DELEGATE_OneParam(FDelegate, const TSharedRef<SHorizontalBox>&);
 
 	FName Hook;
 	EExtensionHook::Position HookPosition = EExtensionHook::Before;
@@ -44,7 +44,7 @@ public:
 		const TSharedPtr<FUICommandList>& InCommands, const FMediaImageStatusBarExtension::FDelegate& InDelegate);
 
 	/** Used by the viewport to add extensions to the status bar. */
-	void Apply(FName InExtensionHook, EExtensionHook::Position InHookPosition, SHorizontalBox& InHorizontalBox) const;
+	void Apply(FName InExtensionHook, EExtensionHook::Position InHookPosition, const TSharedRef<SHorizontalBox>& InHorizontalBox) const;
 
 private:
 	/** List of extensions that get applied to the viewport status bar. */
