Media Viewer: Media controls now correctly fade after 2s

diff --git a/Engine/Plugins/Experimental/MediaViewer/Source/MediaViewer/Private/ImageViewers/MediaSourceImageViewer.cpp b/Engine/Plugins/Experimental/MediaViewer/Source/MediaViewer/Private/ImageViewers/MediaSourceImageViewer.cpp
--- a/Engine/Plugins/Experimental/MediaViewer/Source/MediaViewer/Private/ImageViewers/MediaSourceImageViewer.cpp
+++ b/Engine/Plugins/Experimental/MediaViewer/Source/MediaViewer/Private/ImageViewers/MediaSourceImageViewer.cpp
@@ -414,14 +414,14 @@ void FMediaSourceImageViewer::AddReferencedObjects(FReferenceCollector& Collecto
 	}
 }
 
-TSharedPtr<SWidget> FMediaSourceImageViewer::GetOverlayWidget(EMediaImageViewerPosition InPosition)
+TSharedPtr<SWidget> FMediaSourceImageViewer::GetOverlayWidget(EMediaImageViewerPosition InPosition, const TSharedPtr<SMediaViewerTab>& InViewerTab)
 {
 	return SNew(SBox)
 		.HAlign(HAlign_Fill)
 		.VAlign(VAlign_Bottom)
 		.Padding(5.f)
 		[
-			SNew(SMediaSourceOverlay, SharedThis(this))
+			SNew(SMediaSourceOverlay, SharedThis(this), InViewerTab)
 		];
 }
 
diff --git a/Engine/Plugins/Experimental/MediaViewer/Source/MediaViewer/Private/ImageViewers/MediaSourceImageViewer.h b/Engine/Plugins/Experimental/MediaViewer/Source/MediaViewer/Private/ImageViewers/MediaSourceImageViewer.h
--- a/Engine/Plugins/Experimental/MediaViewer/Source/MediaViewer/Private/ImageViewers/MediaSourceImageViewer.h
+++ b/Engine/Plugins/Experimental/MediaViewer/Source/MediaViewer/Private/ImageViewers/MediaSourceImageViewer.h
@@ -121,7 +121,7 @@ public:
 	virtual TSharedPtr<FMediaViewerLibraryItem> CreateLibraryItem() const override;
 	virtual TOptional<TVariant<FColor, FLinearColor>> GetPixelColor(const FIntPoint& InPixelCoords, int32 InMipLevel) const override;
 	virtual TSharedPtr<FStructOnScope> GetCustomSettingsOnScope() const override;
-	virtual TSharedPtr<SWidget> GetOverlayWidget(EMediaImageViewerPosition InPosition) override;
+	virtual TSharedPtr<SWidget> GetOverlayWidget(EMediaImageViewerPosition InPosition, const TSharedPtr<SMediaViewerTab>& InViewerTab) override;
 	virtual void ExtendStatusBar(UE::MediaViewer::FMediaImageStatusBarExtender& InOutStatusBarExtender) override;
 	//~ End FMediaImageViewer
 
diff --git a/Engine/Plugins/Experimental/MediaViewer/Source/MediaViewer/Private/Widgets/MediaViewerDelegates.h b/Engine/Plugins/Experimental/MediaViewer/Source/MediaViewer/Private/Widgets/MediaViewerDelegates.h
--- a/Engine/Plugins/Experimental/MediaViewer/Source/MediaViewer/Private/Widgets/MediaViewerDelegates.h
+++ b/Engine/Plugins/Experimental/MediaViewer/Source/MediaViewer/Private/Widgets/MediaViewerDelegates.h
@@ -16,6 +16,7 @@ struct FMediaViewerSettings;
 namespace UE::MediaViewer
 {
 	class FMediaImageViewer;
+	class SMediaViewerTab;
 }
 
 namespace UE::MediaViewer::Private
@@ -40,6 +41,7 @@ struct FMediaViewerDelegates
 	DECLARE_DELEGATE_RetVal(FMediaViewerSettings*, FGetSettings)
 	DECLARE_DELEGATE_RetVal(TSharedRef<FMediaViewerLibrary>, FGetLibrary)
 	DECLARE_DELEGATE_RetVal(TSharedPtr<FUICommandList>, FGetCommandList)
+	DECLARE_DELEGATE_RetVal(TSharedPtr<SMediaViewerTab>, FGetTab)
 
 	DECLARE_DELEGATE_RetVal_OneParam(TSharedPtr<FMediaImageViewer>, FGetImageViewer, EMediaImageViewerPosition)
 	DECLARE_DELEGATE_TwoParams(FSetImageViewer, EMediaImageViewerPosition, const TSharedRef<FMediaImageViewer>&)
@@ -127,6 +129,12 @@ struct FMediaViewerDelegates
 	/** Returns the command list for the general view. */
 	FGetCommandList GetCommandList;
 
+	/** Returns true if the mouse is over the viewer. */
+	FGetBool IsOverViewer;
+
+	/** Returns the tab using the viewer. */
+	FGetTab GetTab;
+
 	//////////////////////////
 	// Per-viewer delegates //
 	//////////////////////////
diff --git a/Engine/Plugins/Experimental/MediaViewer/Source/MediaViewer/Private/Widgets/SMediaImageViewerOverlay.cpp b/Engine/Plugins/Experimental/MediaViewer/Source/MediaViewer/Private/Widgets/SMediaImageViewerOverlay.cpp
--- a/Engine/Plugins/Experimental/MediaViewer/Source/MediaViewer/Private/Widgets/SMediaImageViewerOverlay.cpp
+++ b/Engine/Plugins/Experimental/MediaViewer/Source/MediaViewer/Private/Widgets/SMediaImageViewerOverlay.cpp
@@ -433,7 +433,9 @@ TSharedRef<SWidget> SMediaImageViewerOverlay::CreateOverlay(bool bInComparisonVi
 
 	TSharedRef<SOverlay> Overlay = SNew(SOverlay);
 
-	TSharedPtr<SWidget> ImageViewerOverlay = ImageViewer.IsValid() ? ImageViewer->GetOverlayWidget(Position) : nullptr;
+	TSharedPtr<SWidget> ImageViewerOverlay = ImageViewer.IsValid() 
+		? ImageViewer->GetOverlayWidget(Position, Delegates->GetTab.Execute()) 
+		: nullptr;
 
 	if (bInComparisonView)
 	{
diff --git a/Engine/Plugins/Experimental/MediaViewer/Source/MediaViewer/Private/Widgets/SMediaSourceOverlay.cpp b/Engine/Plugins/Experimental/MediaViewer/Source/MediaViewer/Private/Widgets/SMediaSourceOverlay.cpp
--- a/Engine/Plugins/Experimental/MediaViewer/Source/MediaViewer/Private/Widgets/SMediaSourceOverlay.cpp
+++ b/Engine/Plugins/Experimental/MediaViewer/Source/MediaViewer/Private/Widgets/SMediaSourceOverlay.cpp
@@ -19,7 +19,11 @@
 #include "Widgets/Images/SImage.h"
 #include "Widgets/Input/SButton.h"
 #include "Widgets/Layout/SBorder.h"
+#include "Widgets/Layout/SBox.h"
+#include "Widgets/MediaViewerDelegates.h"
 #include "Widgets/SBoxPanel.h"
+#include "Widgets/SMediaViewer.h"
+#include "Widgets/SMediaViewerTab.h"
 #include "Widgets/SNullWidget.h"
 #include "Widgets/Text/STextBlock.h"
 
@@ -28,40 +32,80 @@
 namespace UE::MediaViewer::Private
 {
 
+constexpr double FadeTime = 2.0;
+
 void SMediaSourceOverlay::PrivateRegisterAttributes(FSlateAttributeDescriptor::FInitializer&)
 {
 }
 
-void SMediaSourceOverlay::Construct(const FArguments& InArgs, const TSharedRef<FMediaSourceImageViewer>& InImageViewer)
+void SMediaSourceOverlay::Construct(const FArguments& InArgs, const TSharedRef<FMediaSourceImageViewer>& InImageViewer,
+	const TSharedPtr<SMediaViewerTab>& InViewerTab)
 {
 	ImageViewerWeak = InImageViewer;
 
+	if (InViewerTab.IsValid())
+	{
+		Delegates = InViewerTab->GetViewer()->GetDelegates();
+	}
+
 	TrySetFrameRate();
 
 	ChildSlot
 	[
-		SNew(SBorder)
-		.Padding(5.f)
-		.BorderImage(FAppStyle::GetBrush("ToolTip.Background"))
-		.BorderBackgroundColor(FLinearColor(1.f, 1.f, 1.f, 0.75f))
+		SAssignNew(Container, SBox)
 		[
-			SNew(SVerticalBox)		
-			+ SVerticalBox::Slot()
-			.AutoHeight()
-			[
-				CreateSlider()
-			]
-			+ SVerticalBox::Slot()
-			.AutoHeight()
-			.HAlign(HAlign_Fill)
-			.Padding(0.f, 2.f, 0.f, 0.f)
+			SNew(SBorder)
+			.Padding(5.f)
+			.BorderImage(FAppStyle::GetBrush("ToolTip.Background"))
+			.BorderBackgroundColor(FLinearColor(1.f, 1.f, 1.f, 0.75f))
 			[
-				CreateControls()
+				SNew(SVerticalBox)		
+				+ SVerticalBox::Slot()
+				.AutoHeight()
+				[
+					CreateSlider()
+				]
+				+ SVerticalBox::Slot()
+				.AutoHeight()
+				.HAlign(HAlign_Fill)
+				.Padding(0.f, 2.f, 0.f, 0.f)
+				[
+					CreateControls()
+				]
 			]
 		]
 	];
 }
 
+void SMediaSourceOverlay::Tick(const FGeometry& InAllottedGeometry, const double InCurrentTime, const float InDeltaTime)
+{
+	SCompoundWidget::Tick(InAllottedGeometry, InCurrentTime, InDeltaTime);
+
+	if (!Container.IsValid())
+	{
+		return;
+	}
+
+	const bool bMouseOver = Delegates.IsValid() ? Delegates->IsOverViewer.Execute() : true;
+
+	if (LastInteractionTime < 0 || bMouseOver)
+	{
+		LastInteractionTime = InCurrentTime;
+
+		if (Container->GetVisibility() != EVisibility::Visible)
+		{
+			Container->SetVisibility(EVisibility::Visible);
+		}
+	}
+	else if (InCurrentTime > (LastInteractionTime + FadeTime))
+	{
+		if (Container->GetVisibility() != EVisibility::Hidden)
+		{
+			Container->SetVisibility(EVisibility::Hidden);
+		}
+	}
+}
+
 UMediaStream* SMediaSourceOverlay::GetMediaStream() const
 {
 	if (TSharedPtr<FMediaSourceImageViewer> ImageViewer = ImageViewerWeak.Pin())
diff --git a/Engine/Plugins/Experimental/MediaViewer/Source/MediaViewer/Private/Widgets/SMediaSourceOverlay.h b/Engine/Plugins/Experimental/MediaViewer/Source/MediaViewer/Private/Widgets/SMediaSourceOverlay.h
--- a/Engine/Plugins/Experimental/MediaViewer/Source/MediaViewer/Private/Widgets/SMediaSourceOverlay.h
+++ b/Engine/Plugins/Experimental/MediaViewer/Source/MediaViewer/Private/Widgets/SMediaSourceOverlay.h
@@ -8,14 +8,21 @@
 #include "Misc/Optional.h"
 
 class IMediaStreamPlayer;
+class SBox;
 class UMediaPlayer;
 class UMediaStream;
 struct FSlateBrush;
 
+namespace UE::MediaViewer
+{
+	class SMediaViewerTab;
+}
+
 namespace UE::MediaViewer::Private
 {
 
 class FMediaSourceImageViewer;
+struct FMediaViewerDelegates;
 	
 class SMediaSourceOverlay : public SCompoundWidget
 {
@@ -28,14 +35,24 @@ class SMediaSourceOverlay : public SCompoundWidget
 public:
 	virtual ~SMediaSourceOverlay() override = default;
 
-	void Construct(const FArguments& InArgs, const TSharedRef<FMediaSourceImageViewer>& InImageViewer);
+	void Construct(const FArguments& InArgs, const TSharedRef<FMediaSourceImageViewer>& InImageViewer, 
+		const TSharedPtr<SMediaViewerTab>& InViewerTab);
+
+	//~ Begin SWidget
+	virtual void Tick(const FGeometry& InAllottedGeometry, const double InCurrentTime, const float InDeltaTime) override;
+	//~ End SWidget
 
 protected:
 	TWeakPtr<FMediaSourceImageViewer> ImageViewerWeak;
+	TSharedPtr<FMediaViewerDelegates> Delegates;
+	TSharedPtr<SBox> Container;
 	TOptional<float> FrameRateFloat;
 	mutable TOptional<FText> TotalFrames;
 	mutable TOptional<FText> TotalTime;
 
+	/** The last time the mouse was moved over the widget. */
+	double LastInteractionTime = -1.0;
+
 	UMediaStream* GetMediaStream() const;
 
 	IMediaStreamPlayer* GetMediaStreamPlayer() const;
diff --git a/Engine/Plugins/Experimental/MediaViewer/Source/MediaViewer/Private/Widgets/SMediaViewer.cpp b/Engine/Plugins/Experimental/MediaViewer/Source/MediaViewer/Private/Widgets/SMediaViewer.cpp
--- a/Engine/Plugins/Experimental/MediaViewer/Source/MediaViewer/Private/Widgets/SMediaViewer.cpp
+++ b/Engine/Plugins/Experimental/MediaViewer/Source/MediaViewer/Private/Widgets/SMediaViewer.cpp
@@ -47,10 +47,11 @@ void SMediaViewer::PrivateRegisterAttributes(FSlateAttributeDescriptor::FInitial
 {
 }
 
-void SMediaViewer::Construct(const FArguments& InArgs, const FMediaViewerArgs& InMediaViewerArgs,
+void SMediaViewer::Construct(const FArguments& InArgs, const TSharedRef<SMediaViewerTab>& InTab, const FMediaViewerArgs& InMediaViewerArgs,
 	const TSharedRef<FMediaImageViewer>& InImageViewerFirst, const TSharedRef<FMediaImageViewer>& InImageViewerSecond)
 {
 	MediaViewerArgs = InMediaViewerArgs;
+	Tab = InTab;
 	ImageViewers[static_cast<int32>(EMediaImageViewerPosition::First)] = InImageViewerFirst;
 	ImageViewers[static_cast<int32>(EMediaImageViewerPosition::Second)] = InImageViewerSecond;
 	ActiveView = EMediaImageViewerActivePosition::Single;
@@ -518,6 +519,8 @@ void SMediaViewer::CreateDelegates()
 	Delegates->GetLibrary.BindSPLambda(this, [this] { return Library->GetLibrary(); });
 	Delegates->RefreshView.BindSPLambda(this, [this] { InvalidateView(); });
 	Delegates->GetCommandList.BindSPLambda(this, [this] { return CommandList; });
+	Delegates->IsOverViewer.BindSPLambda(this, [this] { return IsHovered(); });
+	Delegates->GetTab.BindSPLambda(this, [this] { return Tab.Pin(); });
 
 	// Per panel delegates
 	Delegates->GetImageViewer.BindSP(this, &SMediaViewer::GetImageViewer);
@@ -692,6 +695,11 @@ FVector2D SMediaViewer::GetLocalCursorPosition() const
 	return CursorLocalPosition;
 }
 
+TSharedRef<FMediaViewerDelegates> SMediaViewer::GetDelegates() const
+{
+	return Delegates.ToSharedRef();
+}
+
 void SMediaViewer::InvalidateView()
 {
 	bInvalidated = true;
diff --git a/Engine/Plugins/Experimental/MediaViewer/Source/MediaViewer/Private/Widgets/SMediaViewer.h b/Engine/Plugins/Experimental/MediaViewer/Source/MediaViewer/Private/Widgets/SMediaViewer.h
--- a/Engine/Plugins/Experimental/MediaViewer/Source/MediaViewer/Private/Widgets/SMediaViewer.h
+++ b/Engine/Plugins/Experimental/MediaViewer/Source/MediaViewer/Private/Widgets/SMediaViewer.h
@@ -19,6 +19,7 @@ namespace UE::MediaViewer
 {
 	class FMediaImageStatusBarExtender;
 	class FMediaImageViewer;
+	class SMediaViewerTab;
 
 	namespace Private
 	{
@@ -55,7 +56,7 @@ public:
 
 	virtual ~SMediaViewer() override;
 
-	void Construct(const FArguments& InArgs, const FMediaViewerArgs& InMediaViewerArgs,
+	void Construct(const FArguments& InArgs, const TSharedRef<SMediaViewerTab>& InTab, const FMediaViewerArgs& InMediaViewerArgs,
 		const TSharedRef<FMediaImageViewer>& InImageViewerFirst, const TSharedRef<FMediaImageViewer>& InImageViewerSecond);
 
 	const FMediaViewerArgs& GetArgs() const { return MediaViewerArgs; }
@@ -88,6 +89,9 @@ public:
 	/** Triggers a recreation of the view on tick. */
 	void InvalidateView();
 
+	/** Returns the bound delegates of this viewer. */
+	TSharedRef<FMediaViewerDelegates> GetDelegates() const;
+
 	//~ Begin SWidget
 	virtual void Tick(const FGeometry& AllottedGeometry, const double InCurrentTime, const float InDeltaTime) override;
 	virtual int32 OnPaint(const FPaintArgs& InArgs, const FGeometry& InAllottedGeometry, const FSlateRect& InMyCullingRect, 
@@ -103,6 +107,8 @@ public:
 protected:
 	static const FSlateColorBrush BackgroundColorBrush;
 
+	/** Owner of this viewer. */
+	TWeakPtr<SMediaViewerTab> Tab;
 	/** Construction args (unchangeable) */
 	FMediaViewerArgs MediaViewerArgs;
 	/** Mutable settings */
diff --git a/Engine/Plugins/Experimental/MediaViewer/Source/MediaViewer/Private/Widgets/SMediaViewerTab.cpp b/Engine/Plugins/Experimental/MediaViewer/Source/MediaViewer/Private/Widgets/SMediaViewerTab.cpp
--- a/Engine/Plugins/Experimental/MediaViewer/Source/MediaViewer/Private/Widgets/SMediaViewerTab.cpp
+++ b/Engine/Plugins/Experimental/MediaViewer/Source/MediaViewer/Private/Widgets/SMediaViewerTab.cpp
@@ -38,7 +38,7 @@ void SMediaViewerTab::Construct(const FArguments& InArgs, const FMediaViewerArgs
 
 	ChildSlot
 	[
-		SAssignNew(Viewer, Private::SMediaViewer, InMediaViewerArgs, ImageViewerLeft.ToSharedRef(), ImageViewerRight.ToSharedRef())
+		SAssignNew(Viewer, Private::SMediaViewer, SharedThis(this), InMediaViewerArgs, ImageViewerLeft.ToSharedRef(), ImageViewerRight.ToSharedRef())
 	];
 }
 
@@ -67,6 +67,11 @@ void SMediaViewerTab::SetImageViewer(EMediaImageViewerPosition InPosition, TShar
 	Viewer->SetImageViewer(InPosition, InImageViewer.ToSharedRef());
 }
 
+TSharedRef<Private::SMediaViewer> SMediaViewerTab::GetViewer() const
+{
+	return Viewer.ToSharedRef();
+}
+
 } // UE::MediaViewer
 
 #undef LOCTEXT_NAMESPACE
diff --git a/Engine/Plugins/Experimental/MediaViewer/Source/MediaViewer/Public/ImageViewer/MediaImageViewer.h b/Engine/Plugins/Experimental/MediaViewer/Source/MediaViewer/Public/ImageViewer/MediaImageViewer.h
--- a/Engine/Plugins/Experimental/MediaViewer/Source/MediaViewer/Public/ImageViewer/MediaImageViewer.h
+++ b/Engine/Plugins/Experimental/MediaViewer/Source/MediaViewer/Public/ImageViewer/MediaImageViewer.h
@@ -106,6 +106,7 @@ namespace UE::MediaViewer
 {
 
 class FMediaImageStatusBarExtender;
+class SMediaViewerTab;
 	
 /** Information about the image viewer's image. */
 struct FMediaImageViewerInfo
@@ -220,7 +221,7 @@ public:
 	 * The above alignment will be used for the widget's slot.
 	 * If no overlay is required, return nullptr.
 	 */
-	virtual TSharedPtr<SWidget> GetOverlayWidget(EMediaImageViewerPosition InPosition) { return nullptr; }
+	virtual TSharedPtr<SWidget> GetOverlayWidget(EMediaImageViewerPosition InPosition, const TSharedPtr<SMediaViewerTab>& InViewerTab) { return nullptr; }
 
 	/** 
 	 * Responds to tracking started in this viewer's viewport. 
diff --git a/Engine/Plugins/Experimental/MediaViewer/Source/MediaViewer/Public/Widgets/SMediaViewerTab.h b/Engine/Plugins/Experimental/MediaViewer/Source/MediaViewer/Public/Widgets/SMediaViewerTab.h
--- a/Engine/Plugins/Experimental/MediaViewer/Source/MediaViewer/Public/Widgets/SMediaViewerTab.h
+++ b/Engine/Plugins/Experimental/MediaViewer/Source/MediaViewer/Public/Widgets/SMediaViewerTab.h
@@ -44,6 +44,8 @@ public:
 
 	MEDIAVIEWER_API void SetImageViewer(EMediaImageViewerPosition InPosition, TSharedPtr<FMediaImageViewer> InImageViewer);
 
+	TSharedRef<Private::SMediaViewer> GetViewer() const;
+
 protected:
 	TSharedPtr<Private::SMediaViewer> Viewer;
 };
