Media Viewer: Fixed history not saving

diff --git a/Engine/Plugins/Experimental/MediaViewer/Source/MediaViewer/Private/Library/MediaViewerLibraryIni.h b/Engine/Plugins/Experimental/MediaViewer/Source/MediaViewer/Private/Library/MediaViewerLibraryIni.h
--- a/Engine/Plugins/Experimental/MediaViewer/Source/MediaViewer/Private/Library/MediaViewerLibraryIni.h
+++ b/Engine/Plugins/Experimental/MediaViewer/Source/MediaViewer/Private/Library/MediaViewerLibraryIni.h
@@ -44,9 +44,9 @@ public:
 	bool HasItem(const FGuid& InItemId) const;
 
 protected:
-	UPROPERTY()
+	UPROPERTY(Config)
 	TArray<FMediaViewerLibraryGroup> Groups;
 
-	UPROPERTY()
+	UPROPERTY(Config)
 	TArray<FMediaViewerLibraryItemData> Items;
 };
