Media Stream: Media texture updates
- Media textures (and all uobject-based assets) now check to see if they are in an actor and use that actor's label/name as their display name.
- Removed the "saveable" icon.
- Unused media textures are no longer displayed.
- Text in the library now clips.

diff --git a/Engine/Plugins/Experimental/MediaViewer/Source/MediaViewer/Private/ImageViewers/MaterialInterfaceImageViewer.cpp b/Engine/Plugins/Experimental/MediaViewer/Source/MediaViewer/Private/ImageViewers/MaterialInterfaceImageViewer.cpp
--- a/Engine/Plugins/Experimental/MediaViewer/Source/MediaViewer/Private/ImageViewers/MaterialInterfaceImageViewer.cpp
+++ b/Engine/Plugins/Experimental/MediaViewer/Source/MediaViewer/Private/ImageViewers/MaterialInterfaceImageViewer.cpp
@@ -75,7 +75,7 @@ TSharedPtr<FMediaViewerLibraryItem> FMaterialInterfaceImageViewer::FFactory::Cre
 	if (UMaterialInterface* MaterialInterface = Cast<UMaterialInterface>(InObject))
 	{
 		return MakeShared<FMaterialInterfaceImageViewer::FItem>(
-			FText::FromName(MaterialInterface->GetFName()),
+			FMediaImageViewer::GetObjectDisplayName(MaterialInterface),
 			FText::Format(LOCTEXT("ToolTipFromat", "{0} [Material]"), FText::FromString(MaterialInterface->GetPathName())),
 			MaterialInterface->HasAnyFlags(EObjectFlags::RF_Transient) || MaterialInterface->IsIn(GetTransientPackage()),
 			MaterialInterface
@@ -211,7 +211,7 @@ FMaterialInterfaceImageViewer::FMaterialInterfaceImageViewer(const FGuid& InId,
 		InId,
 		FIntPoint(1, 1),
 		1,
-		FText::FromName(InMaterialInterface->GetFName())
+		FMediaImageViewer::GetObjectDisplayName(InMaterialInterface)
 	})
 {
 	DrawEffects |= ESlateDrawEffect::PreMultipliedAlpha | ESlateDrawEffect::NoGamma;
@@ -242,7 +242,7 @@ TSharedPtr<FMediaViewerLibraryItem> FMaterialInterfaceImageViewer::CreateLibrary
 
 	return MakeShared<FItem>(
 		ImageInfo.Id,
-		FText::FromName(MaterialInterface->GetFName()),
+		FMediaImageViewer::GetObjectDisplayName(MaterialInterface),
 		FText::Format(LOCTEXT("ToolTipFromat", "{0} [Material]"), FText::FromString(MaterialInterface->GetPathName())),
 		MaterialInterface->HasAnyFlags(EObjectFlags::RF_Transient) || MaterialInterface->IsIn(GetTransientPackage()),
 		MaterialInterface
diff --git a/Engine/Plugins/Experimental/MediaViewer/Source/MediaViewer/Private/ImageViewers/MediaImageViewer.cpp b/Engine/Plugins/Experimental/MediaViewer/Source/MediaViewer/Private/ImageViewers/MediaImageViewer.cpp
--- a/Engine/Plugins/Experimental/MediaViewer/Source/MediaViewer/Private/ImageViewers/MediaImageViewer.cpp
+++ b/Engine/Plugins/Experimental/MediaViewer/Source/MediaViewer/Private/ImageViewers/MediaImageViewer.cpp
@@ -4,6 +4,7 @@
 
 #include "Engine/Texture.h"
 #include "Engine/TextureRenderTarget2D.h"
+#include "GameFramework/Actor.h"
 #include "Math/UnrealMathUtility.h"
 #include "Math/Vector2D.h"
 #include "Misc/TVariant.h"
@@ -14,6 +15,28 @@ namespace UE::MediaViewer
 
 const FSlateColorBrush FMediaImageViewer::BackgroundColorBrush(FLinearColor::White);
 
+MEDIAVIEWER_API FText FMediaImageViewer::GetObjectDisplayName(const UObject* InObject)
+{
+	if (!InObject)
+	{
+		return FText::GetEmpty();
+	}
+
+	if (AActor* Actor = InObject->GetTypedOuter<AActor>())
+	{
+		const FString ActorLabel = Actor->GetActorLabel(/* Create if none */ false);
+
+		if (ActorLabel.IsEmpty())
+		{
+			return FText::FromName(Actor->GetFName());
+		}
+
+		return FText::FromString(ActorLabel);
+	}
+
+	return FText::FromName(InObject->GetFName());
+}
+
 FMediaImageViewer::FMediaImageViewer(const FMediaImageViewerInfo& InImageInfo)
 	: ImageInfo(InImageInfo)
 	, BackgroundImageBrush(static_cast<UObject*>(nullptr), FVector2D(1))
diff --git a/Engine/Plugins/Experimental/MediaViewer/Source/MediaViewer/Private/ImageViewers/MediaSourceImageViewer.cpp b/Engine/Plugins/Experimental/MediaViewer/Source/MediaViewer/Private/ImageViewers/MediaSourceImageViewer.cpp
--- a/Engine/Plugins/Experimental/MediaViewer/Source/MediaViewer/Private/ImageViewers/MediaSourceImageViewer.cpp
+++ b/Engine/Plugins/Experimental/MediaViewer/Source/MediaViewer/Private/ImageViewers/MediaSourceImageViewer.cpp
@@ -79,7 +79,7 @@ TSharedPtr<FMediaImageViewer> FMediaSourceImageViewer::FFactory::CreateImageView
 {
 	if (UMediaSource* MediaSource = Cast<UMediaSource>(InObject))
 	{
-		return MakeShared<FMediaSourceImageViewer>(MediaSource, FText::FromName(MediaSource->GetFName()));
+		return MakeShared<FMediaSourceImageViewer>(MediaSource, FMediaImageViewer::GetObjectDisplayName(MediaSource));
 	}
 
 	return nullptr;
@@ -90,7 +90,7 @@ TSharedPtr<FMediaViewerLibraryItem> FMediaSourceImageViewer::FFactory::CreateLib
 	if (UMediaSource* MediaSource = Cast<UMediaSource>(InObject))
 	{
 		return MakeShared<FMediaSourceImageViewer::FAssetItem>(
-			FText::FromName(MediaSource->GetFName()),
+			FMediaImageViewer::GetObjectDisplayName(MediaSource),
 			FText::Format(LOCTEXT("FactoryToolTipFormat", "{0} [Media Source]"), FText::FromString(MediaSource->GetPathName())),
 			MediaSource->HasAnyFlags(EObjectFlags::RF_Transient) || MediaSource->IsIn(GetTransientPackage()),
 			MediaSource
@@ -182,7 +182,7 @@ TSharedPtr<FMediaImageViewer> FMediaSourceImageViewer::FAssetItem::CreateImageVi
 		return nullptr;
 	}
 
-	const FText DisplayName = FText::FromName(MediaSource->GetFName());
+	const FText DisplayName = FMediaImageViewer::GetObjectDisplayName(MediaSource);
 
 	if (Id.IsValid())
 	{
@@ -329,7 +329,7 @@ TSharedPtr<FMediaViewerLibraryItem> FMediaSourceImageViewer::CreateLibraryItem()
 	{
 		return MakeShared<FAssetItem>(
 			ImageInfo.Id,
-			FText::FromName(MediaSource->GetFName()),
+			FMediaImageViewer::GetObjectDisplayName(MediaSource),
 			FText::Format(LOCTEXT("LibraryToolTipFormat", "{0} [Media Texture - {1}x{2}]"), FText::FromString(MediaSource->GetPathName()), FText::AsNumber(ImageInfo.Size.X), FText::AsNumber(ImageInfo.Size.Y)),
 			bIsTransient,
 			MediaSource
diff --git a/Engine/Plugins/Experimental/MediaViewer/Source/MediaViewer/Private/ImageViewers/MediaTextureImageViewer.cpp b/Engine/Plugins/Experimental/MediaViewer/Source/MediaViewer/Private/ImageViewers/MediaTextureImageViewer.cpp
--- a/Engine/Plugins/Experimental/MediaViewer/Source/MediaViewer/Private/ImageViewers/MediaTextureImageViewer.cpp
+++ b/Engine/Plugins/Experimental/MediaViewer/Source/MediaViewer/Private/ImageViewers/MediaTextureImageViewer.cpp
@@ -70,7 +70,7 @@ TSharedPtr<FMediaViewerLibraryItem> FMediaTextureImageViewer::FFactory::CreateLi
 	if (UMediaTexture* MediaTexture = Cast<UMediaTexture>(InObject))
 	{
 		return MakeShared<FMediaTextureImageViewer::FItem>(
-			FText::FromName(MediaTexture->GetFName()),
+			FMediaImageViewer::GetObjectDisplayName(InObject),
 			FText::Format(LOCTEXT("ToolTipFromat", "{0} [Media Texture - {1}x{2}]"), FText::FromString(MediaTexture->GetPathName()), FText::AsNumber(MediaTexture->GetSurfaceWidth()), FText::AsNumber(MediaTexture->GetSurfaceHeight())),
 			MediaTexture->HasAnyFlags(EObjectFlags::RF_Transient) || MediaTexture->IsIn(GetTransientPackage()),
 			MediaTexture
@@ -167,7 +167,7 @@ FMediaTextureImageViewer::FMediaTextureImageViewer(const FGuid& InId, TNotNull<U
 		InId,
 		FIntPoint(InMediaTexture->GetSurfaceWidth(), InMediaTexture->GetSurfaceHeight()),
 		InMediaTexture->GetTextureNumMips(),
-		FText::FromName(InMediaTexture->GetFName())
+		FMediaImageViewer::GetObjectDisplayName(InMediaTexture)
 	})
 {
 	MediaTextureSettings.MediaTexture = InMediaTexture;
@@ -197,7 +197,7 @@ TSharedPtr<FMediaViewerLibraryItem> FMediaTextureImageViewer::CreateLibraryItem(
 
 	return MakeShared<FItem>(
 		ImageInfo.Id,
-		FText::FromName(MediaTexture->GetFName()),
+		FMediaImageViewer::GetObjectDisplayName(MediaTexture),
 		FText::Format(LOCTEXT("ToolTipFromat", "{0} [Media Texture - {1}x{2}]"), FText::FromString(MediaTexture->GetPathName()), FText::AsNumber(MediaTexture->GetSurfaceWidth()), FText::AsNumber(MediaTexture->GetSurfaceHeight())),
 		MediaTexture->HasAnyFlags(EObjectFlags::RF_Transient) || MediaTexture->IsIn(GetTransientPackage()),
 		MediaTexture
diff --git a/Engine/Plugins/Experimental/MediaViewer/Source/MediaViewer/Private/ImageViewers/Texture2DImageViewer.cpp b/Engine/Plugins/Experimental/MediaViewer/Source/MediaViewer/Private/ImageViewers/Texture2DImageViewer.cpp
--- a/Engine/Plugins/Experimental/MediaViewer/Source/MediaViewer/Private/ImageViewers/Texture2DImageViewer.cpp
+++ b/Engine/Plugins/Experimental/MediaViewer/Source/MediaViewer/Private/ImageViewers/Texture2DImageViewer.cpp
@@ -68,7 +68,7 @@ TSharedPtr<FMediaViewerLibraryItem> FTexture2DImageViewer::FFactory::CreateLibra
 	if (UTexture2D* Texture = Cast<UTexture2D>(InObject))
 	{
 		return MakeShared<FTexture2DImageViewer::FItem>(
-			FText::FromName(Texture->GetFName()),
+			FMediaImageViewer::GetObjectDisplayName(Texture),
 			FText::Format(LOCTEXT("ToolTipFromat", "{0} [Texture 2D - {1}x{2}]"), FText::FromString(Texture->GetPathName()), FText::AsNumber(Texture->GetSurfaceWidth()), FText::AsNumber(Texture->GetSurfaceHeight())),
 			Texture->HasAnyFlags(EObjectFlags::RF_Transient) || Texture->IsIn(GetTransientPackage()),
 			Texture
@@ -195,7 +195,7 @@ FTexture2DImageViewer::FTexture2DImageViewer(const FGuid& InId, TNotNull<UTextur
 		InId,
 		FIntPoint(InTexture->GetSurfaceWidth(), InTexture->GetSurfaceHeight()),
 		InTexture->GetPlatformData()->Mips.Num(),
-		FText::FromName(InTexture->GetFName())
+		FMediaImageViewer::GetObjectDisplayName(InTexture)
 	})
 	, SampleCache(InTexture, InTexture->GetPixelFormat())
 {
@@ -218,7 +218,7 @@ TSharedPtr<FMediaViewerLibraryItem> FTexture2DImageViewer::CreateLibraryItem() c
 
 	return MakeShared<FItem>(
 		ImageInfo.Id,
-		FText::FromName(Texture->GetFName()),
+		FMediaImageViewer::GetObjectDisplayName(Texture),
 		FText::Format(LOCTEXT("ToolTipFromat", "{0} [Texture 2D - {1}x{2}]"), FText::FromString(Texture->GetPathName()), FText::AsNumber(Texture->GetSurfaceWidth()), FText::AsNumber(Texture->GetSurfaceHeight())),
 		Texture->HasAnyFlags(EObjectFlags::RF_Transient) || Texture->IsIn(GetTransientPackage()),
 		Texture
diff --git a/Engine/Plugins/Experimental/MediaViewer/Source/MediaViewer/Private/ImageViewers/TextureRenderTarget2DImageViewer.cpp b/Engine/Plugins/Experimental/MediaViewer/Source/MediaViewer/Private/ImageViewers/TextureRenderTarget2DImageViewer.cpp
--- a/Engine/Plugins/Experimental/MediaViewer/Source/MediaViewer/Private/ImageViewers/TextureRenderTarget2DImageViewer.cpp
+++ b/Engine/Plugins/Experimental/MediaViewer/Source/MediaViewer/Private/ImageViewers/TextureRenderTarget2DImageViewer.cpp
@@ -67,7 +67,7 @@ TSharedPtr<FMediaViewerLibraryItem> FTextureRenderTarget2DImageViewer::FFactory:
 	if (UTextureRenderTarget2D* RenderTarget = Cast<UTextureRenderTarget2D>(InObject))
 	{
 		return MakeShared<FTextureRenderTarget2DImageViewer::FItem>(
-			FText::FromName(RenderTarget->GetFName()),
+			FMediaImageViewer::GetObjectDisplayName(RenderTarget),
 			FText::Format(LOCTEXT("ToolTipFromat", "{0} [Render Target 2D - {1}x{2}]"), FText::FromString(RenderTarget->GetPathName()), FText::AsNumber(RenderTarget->GetSurfaceWidth()), FText::AsNumber(RenderTarget->GetSurfaceHeight())),
 			RenderTarget->HasAnyFlags(EObjectFlags::RF_Transient) || RenderTarget->IsIn(GetTransientPackage()),
 			RenderTarget
@@ -164,7 +164,7 @@ FTextureRenderTarget2DImageViewer::FTextureRenderTarget2DImageViewer(const FGuid
 		InId,
 		FIntPoint(InRenderTarget->GetSurfaceWidth(), InRenderTarget->GetSurfaceHeight()),
 		InRenderTarget->GetNumMips(),
-		FText::FromName(InRenderTarget->GetFName())
+		FMediaImageViewer::GetObjectDisplayName(InRenderTarget)
 	})
 {
 	DrawEffects |= ESlateDrawEffect::PreMultipliedAlpha | ESlateDrawEffect::NoGamma;
@@ -190,7 +190,7 @@ TSharedPtr<FMediaViewerLibraryItem> FTextureRenderTarget2DImageViewer::CreateLib
 
 	return MakeShared<FItem>(
 		ImageInfo.Id,
-		FText::FromName(RenderTarget->GetFName()),
+		FMediaImageViewer::GetObjectDisplayName(RenderTarget),
 		FText::Format(LOCTEXT("ToolTipFromat", "{0} [Render Target 2D - {1}x{2}]"), FText::FromString(RenderTarget->GetPathName()), FText::AsNumber(RenderTarget->GetSurfaceWidth()), FText::AsNumber(RenderTarget->GetSurfaceHeight())),
 		RenderTarget->HasAnyFlags(EObjectFlags::RF_Transient) || RenderTarget->IsIn(GetTransientPackage()),
 		RenderTarget
diff --git a/Engine/Plugins/Experimental/MediaViewer/Source/MediaViewer/Private/Library/MediaTextureGroup.cpp b/Engine/Plugins/Experimental/MediaViewer/Source/MediaViewer/Private/Library/MediaTextureGroup.cpp
--- a/Engine/Plugins/Experimental/MediaViewer/Source/MediaViewer/Private/Library/MediaTextureGroup.cpp
+++ b/Engine/Plugins/Experimental/MediaViewer/Source/MediaViewer/Private/Library/MediaTextureGroup.cpp
@@ -48,6 +48,13 @@ TArray<TSharedRef<FMediaViewerLibraryItem>> FMediaTextureGroup::GetMediaTextureI
 			continue;
 		}
 
+		constexpr int32 MinimumSize = 3;
+
+		if (MediaTexture->GetSurfaceWidth() < MinimumSize || MediaTexture->GetSurfaceHeight() < MinimumSize)
+		{
+			continue;
+		}
+
 		if (TSharedPtr<FMediaViewerLibraryItem> NewItem = Factory->CreateLibraryItem(TNotNull<UMediaTexture*>(MediaTexture)))
 		{
 			MediaTextureItems.Add(NewItem.ToSharedRef());
diff --git a/Engine/Plugins/Experimental/MediaViewer/Source/MediaViewer/Private/Widgets/MediaViewerLibraryItemDragDropOperation.cpp b/Engine/Plugins/Experimental/MediaViewer/Source/MediaViewer/Private/Widgets/MediaViewerLibraryItemDragDropOperation.cpp
--- a/Engine/Plugins/Experimental/MediaViewer/Source/MediaViewer/Private/Widgets/MediaViewerLibraryItemDragDropOperation.cpp
+++ b/Engine/Plugins/Experimental/MediaViewer/Source/MediaViewer/Private/Widgets/MediaViewerLibraryItemDragDropOperation.cpp
@@ -94,26 +94,9 @@ void FMediaViewerLibraryItemDragDropOperation::CreateDecorator(const TSharedPtr<
 			.HAlign(HAlign_Left)
 			.VAlign(VAlign_Center)
 			[
-				SNew(SHorizontalBox)
-				+ SHorizontalBox::Slot()
-				.AutoWidth()
-				.VAlign(VAlign_Center)
-				[
-					SNew(SImage)
-					.Image(FAppStyle::Get().GetBrush("Icons.Save"))
-					.ColorAndOpacity(FSlateColor(SaveColor))
-					.DesiredSizeOverride(FVector2D(12.f))
-					.ToolTipText(SaveToolTip)
-				]
-				+ SHorizontalBox::Slot()
-				.AutoWidth()
-				.VAlign(VAlign_Center)
-				.Padding(3.f, 0.f, 0.f, 0.f)
-				[
-					SNew(STextBlock)
-					.Text(InItem->GetItemTypeDisplayName())
-					.Font(IDetailLayoutBuilder::GetDetailFont())
-				]
+				SNew(STextBlock)
+				.Text(InItem->GetItemTypeDisplayName())
+				.Font(IDetailLayoutBuilder::GetDetailFont())
 			]
 		]
 	);
diff --git a/Engine/Plugins/Experimental/MediaViewer/Source/MediaViewer/Private/Widgets/SMediaViewerLibraryItem.cpp b/Engine/Plugins/Experimental/MediaViewer/Source/MediaViewer/Private/Widgets/SMediaViewerLibraryItem.cpp
--- a/Engine/Plugins/Experimental/MediaViewer/Source/MediaViewer/Private/Widgets/SMediaViewerLibraryItem.cpp
+++ b/Engine/Plugins/Experimental/MediaViewer/Source/MediaViewer/Private/Widgets/SMediaViewerLibraryItem.cpp
@@ -65,23 +65,11 @@ void SMediaViewerLibraryItem::Construct(const FArguments& InArgs, const TSharedR
 		.Padding(0.f, 0.f, 5.f, 0.f)
 		[
 			SNew(STextBlock)
+			.Clipping(EWidgetClipping::ClipToBounds)
 			.Font(IDetailLayoutBuilder::GetDetailFont())
 			.Text(Item->GetItemTypeDisplayName())
 		];
 
-	if (!Item->IsTransient())
-	{
-		Details->AddSlot()
-			.AutoWidth()
-			.VAlign(VAlign_Center)
-			[
-				SNew(SImage)
-				.Image(FAppStyle::Get().GetBrush("Icons.Save"))
-				.DesiredSizeOverride(FVector2D(12.f))
-				.ToolTipText(LOCTEXT("NotTransientToolTip", "Saved between sessions."))
-			];
-	}
-
 	TSharedRef<SHorizontalBox> Box = SNew(SHorizontalBox)
 		+ SHorizontalBox::Slot()
 		.AutoWidth()
diff --git a/Engine/Plugins/Experimental/MediaViewer/Source/MediaViewer/Public/ImageViewer/MediaImageViewer.h b/Engine/Plugins/Experimental/MediaViewer/Source/MediaViewer/Public/ImageViewer/MediaImageViewer.h
--- a/Engine/Plugins/Experimental/MediaViewer/Source/MediaViewer/Public/ImageViewer/MediaImageViewer.h
+++ b/Engine/Plugins/Experimental/MediaViewer/Source/MediaViewer/Public/ImageViewer/MediaImageViewer.h
@@ -174,6 +174,13 @@ protected:
 	struct FPrivateToken {};
 
 public:
+	/**
+	 * Given an object, tries to fetch a nice display name for it.
+	 * If it's contained within an actor, it will return the actor's label / name.
+	 * Otherwise it returns the object's name.
+	 */
+	MEDIAVIEWER_API static FText GetObjectDisplayName(const UObject* InObject);
+
 	MEDIAVIEWER_API FMediaImageViewer(const FMediaImageViewerInfo& InImageInfo);
 
 	virtual ~FMediaImageViewer() = default;
