Media Viewer: Fixed crash on opening history group
- Groups now check themselves for invalid items when loaded and toggled open.
- Groups and item widgets now correctly instantiate their base table row class when the item/group is invalid.

diff --git a/Engine/Plugins/Experimental/MediaViewer/Source/MediaViewer/Private/Library/MediaViewerLibrary.cpp b/Engine/Plugins/Experimental/MediaViewer/Source/MediaViewer/Private/Library/MediaViewerLibrary.cpp
--- a/Engine/Plugins/Experimental/MediaViewer/Source/MediaViewer/Private/Library/MediaViewerLibrary.cpp
+++ b/Engine/Plugins/Experimental/MediaViewer/Source/MediaViewer/Private/Library/MediaViewerLibrary.cpp
@@ -45,8 +45,6 @@ FMediaViewerLibrary::FMediaViewerLibrary(const FArgs& InArgs)
 		/* Dynamic */ false
 	));
 
-
-
 	Groups.Add(MakeShared<FMediaViewerLibraryGroup>(
 		GetHistoryGroupId(),
 		LOCTEXT("History", "History"),
@@ -459,6 +457,28 @@ IMediaViewerLibrary::FOnChanged::RegistrationType& FMediaViewerLibrary::GetOnCha
 	return OnChangedDelegate;
 }
 
+void FMediaViewerLibrary::RemoveInvalidGroupItems(const FGuid& InGroup)
+{
+	TSharedPtr<FMediaViewerLibraryGroup> Group = GetGroup(InGroup);
+
+	if (!Group.IsValid())
+	{
+		return;
+	}
+
+	const TArray<FGuid>& GroupItems = Group->GetItems();
+
+	for (int32 Index = GroupItems.Num() - 1; Index >= 0; --Index)
+	{
+		TSharedPtr<FMediaViewerLibraryItem> Item = GetItem(GroupItems[Index]);
+
+		if (!Item.IsValid())
+		{
+			Group->RemoveItemAt(Index);
+		}
+	}
+}
+
 void FMediaViewerLibrary::OnChanged(IMediaViewerLibrary::EChangeType InChangeType)
 {
 	OnChangedDelegate.ExecuteIfBound(SharedThis(this), InChangeType);
diff --git a/Engine/Plugins/Experimental/MediaViewer/Source/MediaViewer/Private/Library/MediaViewerLibrary.h b/Engine/Plugins/Experimental/MediaViewer/Source/MediaViewer/Private/Library/MediaViewerLibrary.h
--- a/Engine/Plugins/Experimental/MediaViewer/Source/MediaViewer/Private/Library/MediaViewerLibrary.h
+++ b/Engine/Plugins/Experimental/MediaViewer/Source/MediaViewer/Private/Library/MediaViewerLibrary.h
@@ -39,6 +39,7 @@ public:
 	virtual const TArray<TSharedRef<const FMediaViewerLibraryGroup>>& GetGroups() const override;
 	virtual bool AddGroup(const TSharedRef<FMediaViewerLibraryGroup>& InNewGroup) override;
 	virtual TSharedPtr<FMediaViewerLibraryGroup> GetGroup(const FGuid& InGroupIdToRemove) const override;
+	virtual void RemoveInvalidGroupItems(const FGuid& InGroup) override;
 	virtual bool CanRemoveGroup(const FGuid& InGroupIdToRemove) const override;
 	virtual TSharedPtr<FMediaViewerLibraryGroup> RemoveGroup(const FGuid& InGroupIdToRemove) override;
 	virtual TSharedPtr<FMediaViewerLibraryItem> FindItemByValue(FName InItemType, const FString& InStringValue) const override;
diff --git a/Engine/Plugins/Experimental/MediaViewer/Source/MediaViewer/Private/Library/MediaViewerLibraryIni.cpp b/Engine/Plugins/Experimental/MediaViewer/Source/MediaViewer/Private/Library/MediaViewerLibraryIni.cpp
--- a/Engine/Plugins/Experimental/MediaViewer/Source/MediaViewer/Private/Library/MediaViewerLibraryIni.cpp
+++ b/Engine/Plugins/Experimental/MediaViewer/Source/MediaViewer/Private/Library/MediaViewerLibraryIni.cpp
@@ -84,6 +84,8 @@ void UMediaViewerLibraryIni::LoadLibrary(const TSharedRef<UE::MediaViewer::Priva
 				ExistingItems.Add(ItemId);
 			}
 		}
+
+		InLibrary->RemoveInvalidGroupItems(SavedGroup.GetId());
 	}
 }
 
diff --git a/Engine/Plugins/Experimental/MediaViewer/Source/MediaViewer/Private/Widgets/SMediaViewerLibrary.cpp b/Engine/Plugins/Experimental/MediaViewer/Source/MediaViewer/Private/Widgets/SMediaViewerLibrary.cpp
--- a/Engine/Plugins/Experimental/MediaViewer/Source/MediaViewer/Private/Widgets/SMediaViewerLibrary.cpp
+++ b/Engine/Plugins/Experimental/MediaViewer/Source/MediaViewer/Private/Widgets/SMediaViewerLibrary.cpp
@@ -141,7 +141,10 @@ void SMediaViewerLibrary::OnGetChildren(IMediaViewerLibrary::FGroupItem InParent
 
 	for (const FGuid& ItemId : Group->GetItems())
 	{
-		OutChildren.Add({InParent.GroupId, ItemId});
+		if (Library->GetItem(ItemId).IsValid())
+		{
+			OutChildren.Add({ InParent.GroupId, ItemId });
+		}
 	}
 }
 
diff --git a/Engine/Plugins/Experimental/MediaViewer/Source/MediaViewer/Private/Widgets/SMediaViewerLibraryGroup.cpp b/Engine/Plugins/Experimental/MediaViewer/Source/MediaViewer/Private/Widgets/SMediaViewerLibraryGroup.cpp
--- a/Engine/Plugins/Experimental/MediaViewer/Source/MediaViewer/Private/Widgets/SMediaViewerLibraryGroup.cpp
+++ b/Engine/Plugins/Experimental/MediaViewer/Source/MediaViewer/Private/Widgets/SMediaViewerLibraryGroup.cpp
@@ -19,6 +19,7 @@
 #include "Widgets/Layout/SBox.h"
 #include "Widgets/MediaViewerLibraryItemDragDropOperation.h"
 #include "Widgets/SBoxPanel.h"
+#include "Widgets/SNullWidget.h"
 #include "Widgets/Text/STextBlock.h"
 
 #define LOCTEXT_NAMESPACE "SMediaViewerLibraryGroup"
@@ -42,11 +43,6 @@ void SMediaViewerLibraryGroup::Construct(const FArguments& InArgs, const TShared
 
 	TSharedPtr<FMediaViewerLibraryGroup> Group = InLibrary->GetGroup(InGroupId);
 
-	if (!Group)
-	{
-		return;
-	}
-
 	TSharedRef<SHorizontalBox> Inner = SNew(SHorizontalBox)
 		+ SHorizontalBox::Slot()
 		.FillWidth(1.f)
diff --git a/Engine/Plugins/Experimental/MediaViewer/Source/MediaViewer/Private/Widgets/SMediaViewerLibraryItem.cpp b/Engine/Plugins/Experimental/MediaViewer/Source/MediaViewer/Private/Widgets/SMediaViewerLibraryItem.cpp
--- a/Engine/Plugins/Experimental/MediaViewer/Source/MediaViewer/Private/Widgets/SMediaViewerLibraryItem.cpp
+++ b/Engine/Plugins/Experimental/MediaViewer/Source/MediaViewer/Private/Widgets/SMediaViewerLibraryItem.cpp
@@ -22,6 +22,7 @@
 #include "Widgets/MediaViewerDelegates.h"
 #include "Widgets/MediaViewerLibraryItemDragDropOperation.h"
 #include "Widgets/SBoxPanel.h"
+#include "Widgets/SNullWidget.h"
 #include "Widgets/Text/STextBlock.h"
 
 #define LOCTEXT_NAMESPACE "SMediaViewerLibraryItem"
@@ -45,13 +46,9 @@ void SMediaViewerLibraryItem::Construct(const FArguments& InArgs, const TSharedR
 	GroupItem = InGroupItem;
 	Delegates = InDelegates;
 
+	// Validated by the Library widget
 	Item = InLibrary->GetItem(InGroupItem.ItemId);
 
-	if (!Item.IsValid())
-	{
-		return;
-	}
-
 	TSharedPtr<FMediaViewerLibraryGroup> Group = InLibrary->GetGroup(GroupItem.GroupId);
 
 	SetCursor(EMouseCursor::Hand);
diff --git a/Engine/Plugins/Experimental/MediaViewer/Source/MediaViewer/Public/Library/IMediaViewerLibrary.h b/Engine/Plugins/Experimental/MediaViewer/Source/MediaViewer/Public/Library/IMediaViewerLibrary.h
--- a/Engine/Plugins/Experimental/MediaViewer/Source/MediaViewer/Public/Library/IMediaViewerLibrary.h
+++ b/Engine/Plugins/Experimental/MediaViewer/Source/MediaViewer/Public/Library/IMediaViewerLibrary.h
@@ -91,6 +91,11 @@ public:
 	 */
 	virtual TSharedPtr<FMediaViewerLibraryGroup> GetGroup(const FGuid& InGroupId) const = 0;
 
+	/**
+	 * Scans the group and removes invalid entries.
+	 */
+	virtual void RemoveInvalidGroupItems(const FGuid& InGroup) = 0;
+
 	/**
 	 * Returns true if the group with this id can be removed.
 	 */
