Media Viewer: State saving and restoring
- Added the ability to save and reload the last state.

diff --git a/Engine/Plugins/Experimental/MediaViewer/Source/MediaViewer/Private/ImageViewers/Texture2DImageViewer.cpp b/Engine/Plugins/Experimental/MediaViewer/Source/MediaViewer/Private/ImageViewers/Texture2DImageViewer.cpp
--- a/Engine/Plugins/Experimental/MediaViewer/Source/MediaViewer/Private/ImageViewers/Texture2DImageViewer.cpp
+++ b/Engine/Plugins/Experimental/MediaViewer/Source/MediaViewer/Private/ImageViewers/Texture2DImageViewer.cpp
@@ -4,12 +4,14 @@
 
 #include "AssetRegistry/AssetData.h"
 #include "Brushes/SlateImageBrush.h"
+#include "Engine/Texture.h"
 #include "Engine/Texture2D.h"
 #include "Math/Color.h"
+#include "Math/IntPoint.h"
 #include "Math/Vector2D.h"
 #include "MediaViewerUtils.h"
-#include "Misc/Optional.h"
 #include "Misc/TVariant.h"
+#include "TextureResource.h"
 
 #define LOCTEXT_NAMESPACE "Texture2DImageViewer"
 
@@ -198,13 +200,19 @@ FTexture2DImageViewer::FTexture2DImageViewer(const FGuid& InId, TNotNull<UTextur
 		FMediaImageViewer::GetObjectDisplayName(InTexture)
 	})
 	, SampleCache(InTexture, InTexture->GetPixelFormat())
+	, bValidImageSize(false)
 {
 	TextureSettings.Texture = InTexture;
+	FVector2D ImageSize = FVector2D::ZeroVector;
 
-	Brush = MakeShared<FSlateImageBrush>(
-		InTexture,
-		FVector2D(InTexture->GetSurfaceWidth(), InTexture->GetSurfaceHeight())
-	);
+	if (TOptional<FIntPoint> TextureSize = GetTextureSize())
+	{
+		ImageSize.X = TextureSize->X;
+		ImageSize.Y = TextureSize->Y;
+		bValidImageSize = true;
+	}
+
+	Brush = MakeShared<FSlateImageBrush>(InTexture, ImageSize);
 }
 
 TSharedPtr<FMediaViewerLibraryItem> FTexture2DImageViewer::CreateLibraryItem() const
@@ -269,6 +277,53 @@ void FTexture2DImageViewer::AddReferencedObjects(FReferenceCollector& Collector)
 	);
 }
 
+TOptional<FIntPoint> FTexture2DImageViewer::GetTextureSize() const
+{
+	if (!TextureSettings.Texture)
+	{
+		return {};
+	}
+
+	FTexturePlatformData* PlatformData = TextureSettings.Texture->GetPlatformData();
+
+	if (!PlatformData)
+	{
+		return {};
+	}
+
+	if (FTextureResource* TextureResource = TextureSettings.Texture->GetResource())
+	{
+		if (TextureResource->IsProxy())
+		{
+			return {};
+		}
+	}
+
+	return FIntPoint(TextureSettings.Texture->GetSurfaceWidth(), TextureSettings.Texture->GetSurfaceHeight());
+}
+
+void FTexture2DImageViewer::PaintImage(FMediaImagePaintParams& InPaintParams, const FMediaImagePaintGeometry& InPaintGeometry)
+{
+	if (!bValidImageSize)
+	{
+		if (TOptional<FIntPoint> TextureSize = GetTextureSize())
+		{
+			ImageInfo.Size.X = TextureSettings.Texture->GetSurfaceWidth();
+			ImageInfo.Size.Y = TextureSettings.Texture->GetSurfaceHeight();
+
+			if (Brush.IsValid())
+			{
+				Brush->SetImageSize(FVector2D(
+					TextureSettings.Texture->GetSurfaceWidth(),
+					TextureSettings.Texture->GetSurfaceHeight()
+				));
+			}
+		}
+	}
+
+	FMediaImageViewer::PaintImage(InPaintParams, InPaintGeometry);
+}
+
 FString FTexture2DImageViewer::GetReferencerName() const
 {
 	static const FString ReferencerName = TEXT("FTexture2DImageViewer");
diff --git a/Engine/Plugins/Experimental/MediaViewer/Source/MediaViewer/Private/ImageViewers/Texture2DImageViewer.h b/Engine/Plugins/Experimental/MediaViewer/Source/MediaViewer/Private/ImageViewers/Texture2DImageViewer.h
--- a/Engine/Plugins/Experimental/MediaViewer/Source/MediaViewer/Private/ImageViewers/Texture2DImageViewer.h
+++ b/Engine/Plugins/Experimental/MediaViewer/Source/MediaViewer/Private/ImageViewers/Texture2DImageViewer.h
@@ -8,6 +8,8 @@
 #include "UObject/GCObject.h"
 
 #include "ImageViewers/TextureSampleCache.h"
+#include "Math/MathFwd.h"
+#include "Misc/Optional.h"
 #include "Templates/SharedPointer.h"
 
 #include "Texture2DImageViewer.generated.h"
@@ -95,6 +97,13 @@ public:
 protected:
 	FTexture2DImageViewerSettings TextureSettings;
 	mutable FTextureSampleCache SampleCache;
+	bool bValidImageSize;
+
+	TOptional<FIntPoint> GetTextureSize() const;
+
+	//~ Begin FMediaImageViewer
+	virtual void PaintImage(FMediaImagePaintParams& InPaintParams, const FMediaImagePaintGeometry& InPaintGeometry) override;
+	//~ End FMediaImageViewer
 };
 
 } // UE::MediaViewer::Private
diff --git a/Engine/Plugins/Experimental/MediaViewer/Source/MediaViewer/Private/Library/MediaViewerLibraryIni.cpp b/Engine/Plugins/Experimental/MediaViewer/Source/MediaViewer/Private/Library/MediaViewerLibraryIni.cpp
--- a/Engine/Plugins/Experimental/MediaViewer/Source/MediaViewer/Private/Library/MediaViewerLibraryIni.cpp
+++ b/Engine/Plugins/Experimental/MediaViewer/Source/MediaViewer/Private/Library/MediaViewerLibraryIni.cpp
@@ -42,8 +42,6 @@ void UMediaViewerLibraryIni::SaveLibrary(const TSharedRef<UE::MediaViewer::Priva
 
 		Items.Emplace(Item->GetItemType(), *Item);
 	}
-
-	SaveConfig();
 }
 
 void UMediaViewerLibraryIni::LoadLibrary(const TSharedRef<UE::MediaViewer::Private::FMediaViewerLibrary>& InLibrary) const
@@ -114,3 +112,13 @@ bool UMediaViewerLibraryIni::HasItem(const FGuid& InItemId) const
 
 	return false;
 }
+
+TConstArrayView<FMediaViewerState> UMediaViewerLibraryIni::GetSavedStates() const
+{
+	return SavedStates;
+}
+
+void UMediaViewerLibraryIni::SetSavedStates(const TArray<FMediaViewerState>& InStates)
+{
+	SavedStates = InStates;
+}
diff --git a/Engine/Plugins/Experimental/MediaViewer/Source/MediaViewer/Private/Library/MediaViewerLibraryIni.h b/Engine/Plugins/Experimental/MediaViewer/Source/MediaViewer/Private/Library/MediaViewerLibraryIni.h
--- a/Engine/Plugins/Experimental/MediaViewer/Source/MediaViewer/Private/Library/MediaViewerLibraryIni.h
+++ b/Engine/Plugins/Experimental/MediaViewer/Source/MediaViewer/Private/Library/MediaViewerLibraryIni.h
@@ -6,10 +6,14 @@
 
 #include "Library/MediaViewerLibraryGroup.h"
 #include "Library/MediaViewerLibraryItem.h"
+#include "ImageViewer/MediaImageViewer.h"
 #include "Misc/Guid.h"
+#include "Widgets/MediaViewerSettings.h"
 
 #include "MediaViewerLibraryIni.generated.h"
 
+enum class EMediaImageViewerActivePosition : uint8;
+
 namespace UE::MediaViewer::Private
 {
 	class FMediaViewerLibrary;
@@ -27,6 +31,39 @@ struct FMediaViewerLibraryItemData
 	FMediaViewerLibraryItem Item;
 };
 
+USTRUCT()
+struct FMediaViewerImageState
+{
+	GENERATED_BODY()
+
+	UPROPERTY(Config)
+	FName ImageType;
+
+	UPROPERTY(Config)
+	FString StringValue;
+
+	UPROPERTY(Config)
+	FMediaImageViewerPanelSettings PanelSettings;
+
+	UPROPERTY(Config)
+	FMediaImagePaintSettings PaintSettings;
+};
+
+USTRUCT()
+struct FMediaViewerState
+{
+	GENERATED_BODY()
+
+	UPROPERTY(Config)
+	FMediaViewerSettings ViewerSettings;
+
+	UPROPERTY(Config)
+	EMediaImageViewerActivePosition ActiveView;
+
+	UPROPERTY(Config)
+	TArray<FMediaViewerImageState> Images;
+};
+
 UCLASS(Config = EditorPerProjectUserSettings)
 class UMediaViewerLibraryIni : public UObject
 {
@@ -43,10 +80,17 @@ public:
 
 	bool HasItem(const FGuid& InItemId) const;
 
+	TConstArrayView<FMediaViewerState> GetSavedStates() const;
+
+	void SetSavedStates(const TArray<FMediaViewerState>& InStates);
+
 protected:
 	UPROPERTY(Config)
 	TArray<FMediaViewerLibraryGroup> Groups;
 
 	UPROPERTY(Config)
 	TArray<FMediaViewerLibraryItemData> Items;
+
+	UPROPERTY(Config)
+	TArray<FMediaViewerState> SavedStates;
 };
diff --git a/Engine/Plugins/Experimental/MediaViewer/Source/MediaViewer/Private/MediaViewerModule.cpp b/Engine/Plugins/Experimental/MediaViewer/Source/MediaViewer/Private/MediaViewerModule.cpp
--- a/Engine/Plugins/Experimental/MediaViewer/Source/MediaViewer/Private/MediaViewerModule.cpp
+++ b/Engine/Plugins/Experimental/MediaViewer/Source/MediaViewer/Private/MediaViewerModule.cpp
@@ -22,6 +22,7 @@
 #include "Modules/ModuleManager.h"
 #include "Widgets/Docking/SDockTab.h"
 #include "Widgets/MediaViewerDelegates.h"
+#include "Widgets/SMediaViewer.h"
 #include "Widgets/SMediaViewerLibrary.h"
 #include "Widgets/SMediaViewerTab.h"
 #include "WorkspaceMenuStructure.h"
@@ -356,11 +357,18 @@ void FMediaViewerModule::OnTabClosed(TSharedRef<SDockTab> InDockTab)
 
 	if (Content->GetWidgetClass().GetWidgetType() == SMediaViewerTab::StaticWidgetClass().GetWidgetType())
 	{
-		UMediaViewerLibraryIni::Get().SaveLibrary(
+		TSharedRef<SMediaViewerTab> MediaViewerTab = StaticCastSharedRef<SMediaViewerTab>(Content);
+		MediaViewerTab->GetViewer()->SaveStates();
+
+		UMediaViewerLibraryIni& Ini = UMediaViewerLibraryIni::Get();
+
+		Ini.SaveLibrary(
 			StaticCastSharedRef<FMediaViewerLibrary>(
-				StaticCastSharedRef<SMediaViewerTab>(Content)->GetLibrary()
+				MediaViewerTab->GetLibrary()
 			)
 		);
+
+		Ini.SaveConfig();
 	}
 }
 
diff --git a/Engine/Plugins/Experimental/MediaViewer/Source/MediaViewer/Private/Widgets/MediaViewerSettings.h b/Engine/Plugins/Experimental/MediaViewer/Source/MediaViewer/Private/Widgets/MediaViewerSettings.h
--- a/Engine/Plugins/Experimental/MediaViewer/Source/MediaViewer/Private/Widgets/MediaViewerSettings.h
+++ b/Engine/Plugins/Experimental/MediaViewer/Source/MediaViewer/Private/Widgets/MediaViewerSettings.h
@@ -20,31 +20,32 @@ struct FMediaViewerSettings
 	GENERATED_BODY()
 
 	/** Color painted to the viewer behind the images. */
-	UPROPERTY(EditAnywhere, Category = "Background")
+	UPROPERTY(EditAnywhere, Config, Category = "Background")
 	FLinearColor ClearColor = FStyleColors::Recessed.GetSpecifiedColor();
 
 	/** Texture tiled behind the images, on top of the clear color. */
-	UPROPERTY(EditAnywhere, Category = "Background")
-	TObjectPtr<UTexture2D> Texture = nullptr;
+	UPROPERTY(EditAnywhere, Config, Category = "Background")
+	TSoftObjectPtr<UTexture2D> Texture = nullptr;
 
 	/** Offset of the tiled background texture. */
-	UPROPERTY(EditAnywhere, Category = "Background", meta = (EditCondition = "Texture != nullptr"))
+	UPROPERTY(EditAnywhere, Config, Category = "Background", meta = (EditCondition = "Texture != nullptr"))
 	FVector2D Offset = FVector2D::ZeroVector;
 
 	/** Scale of the tiled background texture. */
-	UPROPERTY(EditAnywhere, Category = "Background", meta = (UIMin = 0.1, ClampMin = 0.1, UIMax = 10, ClampMax = 10,
+	UPROPERTY(EditAnywhere, Config, Category = "Background", meta = (UIMin = 0.1, ClampMin = 0.1, UIMax = 10, ClampMax = 10,
 		EditCondition = "Texture != nullptr"))
 	float Scale = 1.f;
 
 	/** If true, when one image transform changes, the other image transforms will change an equal amount. */
-	UPROPERTY(EditAnywhere, Category = "AB View", DisplayName = "Lock Transforms")
+	UPROPERTY(EditAnywhere, Config, Category = "AB View", DisplayName = "Lock Transforms")
 	bool bAreTransformsLocked = true;
 
 	/** When in AB view, whether the viewers are split horizontally or vertically. */
+	UPROPERTY(Config)
 	TEnumAsByte<EOrientation> ABOrientation = Orient_Horizontal;
 
 	/** Splitter location (0-100) */
-	UPROPERTY(EditAnywhere, Category = "AB View", DisplayName = "AB Splitter Location",
+	UPROPERTY(EditAnywhere, Config, Category = "AB View", DisplayName = "AB Splitter Location",
 		meta = (UIMin = 0, ClampMin = 0, UIMax = 100, ClampMax = 100, Units = "Percent"))
 	float ABSplitterLocation = 50.f;
 
@@ -52,7 +53,7 @@ struct FMediaViewerSettings
 	 * At 100%, the first/A image is rendered only on the left/top side.
 	 * Below 100%, the first/A image is rendered on the entire panel with B overlaid on top of it, translucently.
 	 */
-	UPROPERTY(EditAnywhere, Category = "AB View", DisplayName = "B Image Opacity",
+	UPROPERTY(EditAnywhere, Config, Category = "AB View", DisplayName = "B Image Opacity",
 		meta = (UIMin = 0, ClampMin = 0, UIMax = 100, ClampMax = 100, Units = "Percent"))
 	float SecondImageOpacity = 100.f;
 };
diff --git a/Engine/Plugins/Experimental/MediaViewer/Source/MediaViewer/Private/Widgets/SMediaViewer.cpp b/Engine/Plugins/Experimental/MediaViewer/Source/MediaViewer/Private/Widgets/SMediaViewer.cpp
--- a/Engine/Plugins/Experimental/MediaViewer/Source/MediaViewer/Private/Widgets/SMediaViewer.cpp
+++ b/Engine/Plugins/Experimental/MediaViewer/Source/MediaViewer/Private/Widgets/SMediaViewer.cpp
@@ -9,8 +9,11 @@
 #include "ImageViewer/MediaImageViewer.h"
 #include "ImageViewers/NullImageViewer.h"
 #include "IMediaViewerModule.h"
+#include "Library/MediaViewerLibrary.h"
+#include "Library/MediaViewerLibraryIni.h"
 #include "Library/MediaViewerLibraryGroup.h"
 #include "MediaViewerCommands.h"
+#include "Misc/MessageDialog.h"
 #include "Sidebar/SSidebar.h"
 #include "Sidebar/SSidebarContainer.h"
 #include "Widgets/Colors/SColorBlock.h"
@@ -18,7 +21,6 @@
 #include "Widgets/Layout/SSplitter.h"
 #include "Widgets/MediaViewerDelegates.h"
 #include "Widgets/SMediaImageViewerOverlay.h"
-
 #include "Widgets/SMediaViewerLibrary.h"
 #include "Widgets/SMediaViewerLibraryPrivate.h"
 #include "Widgets/SMediaViewerToolbar.h"
@@ -286,6 +288,19 @@ void SMediaViewer::Tick(const FGeometry& AllottedGeometry, const double InCurren
 		CursorLocalPosition = LocalPosition;
 	}
 
+	constexpr int32 TickCheckValue = 2;
+	constexpr int32 TickStopCheckValue = 3;
+
+	if (TickCount < TickStopCheckValue)
+	{
+		++TickCount;
+
+		if (TickCount == TickCheckValue)
+		{
+			CheckLoadState();
+		}
+	}
+
 	if (bInvalidated)
 	{
 		CreateView();
@@ -324,7 +339,7 @@ int32 SMediaViewer::OnPaint(const FPaintArgs& InArgs, const FGeometry& InAllotte
 
 	++InLayerId;
 
-	if (UTexture* BackgroundTexture = MediaViewerSettings.Texture)
+	if (UTexture* BackgroundTexture = MediaViewerSettings.Texture.LoadSynchronous())
 	{
 		BackgroundTextureBrush.SetResourceObject(BackgroundTexture);
 		BackgroundTextureBrush.ImageSize.X = BackgroundTexture->GetSurfaceWidth();
@@ -700,6 +715,105 @@ TSharedRef<FMediaViewerDelegates> SMediaViewer::GetDelegates() const
 	return Delegates.ToSharedRef();
 }
 
+void SMediaViewer::SaveStates() const
+{
+	FMediaViewerState State;
+	State.ViewerSettings = MediaViewerSettings;
+	State.ActiveView = ActiveView;
+
+	for (int32 Index = 0; Index < static_cast<int32>(EMediaImageViewerPosition::COUNT); ++Index)
+	{
+		bool bAdded = false;
+
+		if (ImageViewers[Index].IsValid())
+		{
+			if (TSharedPtr<FMediaViewerLibraryItem> LibraryItem = ImageViewers[Index]->CreateLibraryItem())
+			{
+				State.Images.Add({
+					LibraryItem->GetItemType(),
+					LibraryItem->GetStringValue(),
+					ImageViewers[Index]->GetPanelSettings(),
+					ImageViewers[Index]->GetPaintSettings()
+				});
+				bAdded = true;
+			}
+		}
+
+		if (!bAdded)
+		{
+			State.Images.Add({});
+		}
+	}
+
+	UMediaViewerLibraryIni& Ini = UMediaViewerLibraryIni::Get();
+	Ini.SetSavedStates({State});
+	Ini.SaveConfig();
+}
+
+void SMediaViewer::LoadState(int32 InIndex)
+{
+	TConstArrayView<FMediaViewerState> SavedStates = UMediaViewerLibraryIni::Get().GetSavedStates();
+
+	if (!SavedStates.IsValidIndex(InIndex))
+	{
+		return;
+	}
+
+	const FMediaViewerState& State = SavedStates[InIndex];
+
+	TArray<TSharedPtr<FMediaViewerLibraryItem>> Items;
+	Items.Reserve(State.Images.Num());
+	bool bHasValidItem = false;
+
+	for (int32 Index = 0; Index < State.Images.Num(); ++Index)
+	{
+		TSharedPtr<FMediaViewerLibraryItem> Item = Library->GetLibrary()->FindItemByValue(
+			State.Images[Index].ImageType, 
+			State.Images[Index].StringValue
+		);
+
+		Items.Add(Item);
+
+		if (Item.IsValid())
+		{
+			bHasValidItem = true;
+		}
+	}
+
+	if (!bHasValidItem)
+	{
+		return;
+	}
+
+	const EAppReturnType::Type Result = FMessageDialog::Open(
+		EAppMsgType::YesNo,
+		LOCTEXT("LoadOldState", "Attempt to open previous images?")
+	);
+
+	if (Result != EAppReturnType::Yes)
+	{
+		return;
+	}
+
+	MediaViewerSettings = State.ViewerSettings;
+	RequestedView = State.ActiveView;
+
+	for (int32 Index = 0; Index < State.Images.Num(); ++Index)
+	{
+		if (Items[Index].IsValid())
+		{
+			if (TSharedPtr<FMediaImageViewer> Viewer = Items[Index]->CreateImageViewer())
+			{
+				ImageViewers[Index] = Viewer.ToSharedRef();
+				Viewer->GetPanelSettings() = State.Images[Index].PanelSettings;
+				Viewer->GetPaintSettings() = State.Images[Index].PaintSettings;
+			}
+		}
+	}
+
+	InvalidateView();
+}
+
 void SMediaViewer::InvalidateView()
 {
 	bInvalidated = true;
@@ -874,6 +988,21 @@ void SMediaViewer::ToggleLockedTransform()
 	MediaViewerSettings.bAreTransformsLocked = !MediaViewerSettings.bAreTransformsLocked;
 }
 
+void SMediaViewer::CheckLoadState()
+{
+	const FGuid& NullImageId = FNullImageViewer::GetNullImageViewer()->GetInfo().Id;
+
+	for (int32 Index = 0; Index < static_cast<int32>(EMediaImageViewerPosition::COUNT); ++Index)
+	{
+		if (ImageViewers[Index].IsValid() && ImageViewers[Index]->GetInfo().Id != NullImageId)
+		{
+			return;
+		}
+	}
+
+	LoadState(0);
+}
+
 } // UE::MediaViewer::Private
 
 #undef LOCTEXT_NAMESPACE
diff --git a/Engine/Plugins/Experimental/MediaViewer/Source/MediaViewer/Private/Widgets/SMediaViewer.h b/Engine/Plugins/Experimental/MediaViewer/Source/MediaViewer/Private/Widgets/SMediaViewer.h
--- a/Engine/Plugins/Experimental/MediaViewer/Source/MediaViewer/Private/Widgets/SMediaViewer.h
+++ b/Engine/Plugins/Experimental/MediaViewer/Source/MediaViewer/Private/Widgets/SMediaViewer.h
@@ -12,6 +12,8 @@
 #include "Widgets/MediaViewerSettings.h"
 #include "Widgets/SBoxPanel.h"
 
+#include "SMediaViewer.generated.h"
+
 class FUICommandList;
 class SViewportToolBar;
 
@@ -32,15 +34,16 @@ namespace UE::MediaViewer
 	}
 }
 
-namespace UE::MediaViewer::Private
-{
-
+UENUM()
 enum class EMediaImageViewerActivePosition : uint8
 {
 	Single,
 	Both
 };
 
+namespace UE::MediaViewer::Private
+{
+
 class SMediaViewer : public SCompoundWidget, public FGCObject
 {
 	SLATE_DECLARE_WIDGET(SMediaViewer, SCompoundWidget)
@@ -92,6 +95,10 @@ public:
 	/** Returns the bound delegates of this viewer. */
 	TSharedRef<FMediaViewerDelegates> GetDelegates() const;
 
+	void SaveStates() const;
+
+	void LoadState(int32 InIndex);
+
 	//~ Begin SWidget
 	virtual void Tick(const FGeometry& AllottedGeometry, const double InCurrentTime, const float InDeltaTime) override;
 	virtual int32 OnPaint(const FPaintArgs& InArgs, const FGeometry& InAllottedGeometry, const FSlateRect& InMyCullingRect, 
@@ -139,6 +146,9 @@ protected:
 	/** Whether to scale to fit on rebuild view. */
 	bool ScaleToFit[static_cast<int32>(EMediaImageViewerPosition::COUNT)];
 
+	/** A check to see if we need to load old states. */
+	int32 TickCount = 0;
+
 	FVector2D CursorLocalPosition;
 
 	mutable FVector2D ViewerSize;
@@ -225,6 +235,8 @@ protected:
 	/** For locking the image viewer transforms. */
 	ECheckBoxState AreLockedTransformEnabled() const;
 	void ToggleLockedTransform();
+
+	void CheckLoadState();
 };
 
 } // UE::MediaViewer::Private
diff --git a/Engine/Plugins/Experimental/MediaViewer/Source/MediaViewer/Public/ImageViewer/MediaImageViewer.h b/Engine/Plugins/Experimental/MediaViewer/Source/MediaViewer/Public/ImageViewer/MediaImageViewer.h
--- a/Engine/Plugins/Experimental/MediaViewer/Source/MediaViewer/Public/ImageViewer/MediaImageViewer.h
+++ b/Engine/Plugins/Experimental/MediaViewer/Source/MediaViewer/Public/ImageViewer/MediaImageViewer.h
@@ -53,14 +53,14 @@ struct FMediaImageViewerPanelSettings
 	 * If this is different to the clear color, it shows where the image is even
 	 * if nothing is drawn.
 	 */
-	UPROPERTY(EditAnywhere, Category = "Panel")
+	UPROPERTY(EditAnywhere, Config, Category = "Panel")
 	TOptional<FLinearColor> BackgroundColor;
 
 	/**
 	 * If set, this is drawn over the top of the background color.
 	 * Example usage: checkered background.
 	 */
-	UPROPERTY(EditAnywhere, Category = "Panel")
+	UPROPERTY(EditAnywhere, Config, Category = "Panel")
 	TSoftObjectPtr<UTexture> BackgroundTexture;
 };
 
@@ -75,30 +75,30 @@ struct FMediaImagePaintSettings
 	 * Y: Vertical (T->B)
 	 * Z: Depth (F->B)
 	 */
-	UPROPERTY(EditAnywhere, Category = "Media")
+	UPROPERTY(EditAnywhere, Config, Category = "Media")
 	FVector Offset = FVector::ZeroVector;
 
 	/** Rotation about the offset, for 3d objects. */
-	UPROPERTY(EditAnywhere, Category = "Media")
+	UPROPERTY(EditAnywhere, Config, Category = "Media")
 	FRotator Rotation = FRotator::ZeroRotator;
 
 	/**
 	 * The scale of the image.
 	 * A scale of 2 will mean the image is twice as big.
 	 */
-	UPROPERTY(EditAnywhere, Category = "Media")
+	UPROPERTY(EditAnywhere, Config, Category = "Media")
 	float Scale = 1.f;
 
 	/**
 	 * The requested mip level. If the mip isn't available, the highest available value will be used.
 	 */
-	UPROPERTY(EditAnywhere, Category = "Media")
+	UPROPERTY(EditAnywhere, Config, Category = "Media")
 	uint8 MipLevel = 0;
 
 	/**
 	 * Color to tint the image brush with.
 	 */
-	UPROPERTY(EditAnywhere, Category = "Media")
+	UPROPERTY(EditAnywhere, Config, Category = "Media")
 	FLinearColor Tint = FLinearColor::White;
 };
 
