Media Viewer: When no image is active, and no drag and drop is active, hint text will be displayed in each active panel.

diff --git a/Engine/Plugins/Experimental/MediaViewer/Source/MediaViewer/Private/Widgets/SMediaImageViewerOverlay.cpp b/Engine/Plugins/Experimental/MediaViewer/Source/MediaViewer/Private/Widgets/SMediaImageViewerOverlay.cpp
--- a/Engine/Plugins/Experimental/MediaViewer/Source/MediaViewer/Private/Widgets/SMediaImageViewerOverlay.cpp
+++ b/Engine/Plugins/Experimental/MediaViewer/Source/MediaViewer/Private/Widgets/SMediaImageViewerOverlay.cpp
@@ -9,6 +9,7 @@
 #include "Framework/MultiBox/MultiBoxBuilder.h"
 #include "HAL/PlatformApplicationMisc.h"
 #include "ImageViewer/MediaImageViewer.h"
+#include "ImageViewers/NullImageViewer.h"
 #include "IMediaViewerModule.h"
 #include "Library/MediaViewerLibrary.h"
 #include "Library/MediaViewerLibraryItem.h"
@@ -21,6 +22,7 @@
 #include "Widgets/SMediaImageViewerStatusBar.h"
 #include "Widgets/SMediaViewer.h"
 #include "Widgets/SMediaViewerDropTarget.h"
+#include "Widgets/Text/STextBlock.h"
 
 #define LOCTEXT_NAMESPACE "SMediaImageViewerOverlay"
 
@@ -437,6 +439,17 @@ TSharedRef<SWidget> SMediaImageViewerOverlay::CreateOverlay(bool bInComparisonVi
 		? ImageViewer->GetOverlayWidget(Position, Delegates->GetTab.Execute()) 
 		: nullptr;
 
+	Overlay->AddSlot()
+		.HAlign(HAlign_Center)
+		.VAlign(VAlign_Center)
+		[
+			SNew(STextBlock)
+			.Visibility(this, &SMediaImageViewerOverlay::HintText_GetVisibility)
+			.Text(LOCTEXT("DropTargetMessage", "Drop supported asset or library item here."))
+			.Font(IDetailLayoutBuilder::GetDetailFont())
+			.TextStyle(FAppStyle::Get(), "HintText")
+		];
+
 	if (bInComparisonView)
 	{
 		Overlay->AddSlot()
@@ -906,6 +919,20 @@ TSharedRef<SWidget> SMediaImageViewerOverlay::CreateMenu()
 	return MenuBuilder.MakeWidget();
 }
 
+EVisibility SMediaImageViewerOverlay::HintText_GetVisibility() const
+{
+	if (FSlateApplication::Get().IsDragDropping())
+	{
+		return EVisibility::Collapsed;
+	}
+
+	TSharedPtr<FMediaImageViewer> ImageViewer = GetImageViewer();
+	
+	return (ImageViewer.IsValid() && ImageViewer->GetInfo().Id != FNullImageViewer::GetNullImageViewer()->GetInfo().Id)
+		? EVisibility::Collapsed
+		: EVisibility::Visible;
+}
+
 } // UE::MediaViewer::Private
 
 #undef LOCTEXT_NAMESPACE
diff --git a/Engine/Plugins/Experimental/MediaViewer/Source/MediaViewer/Private/Widgets/SMediaImageViewerOverlay.h b/Engine/Plugins/Experimental/MediaViewer/Source/MediaViewer/Private/Widgets/SMediaImageViewerOverlay.h
--- a/Engine/Plugins/Experimental/MediaViewer/Source/MediaViewer/Private/Widgets/SMediaImageViewerOverlay.h
+++ b/Engine/Plugins/Experimental/MediaViewer/Source/MediaViewer/Private/Widgets/SMediaImageViewerOverlay.h
@@ -147,6 +147,8 @@ protected:
 	void AddToLibrary();
 
 	TSharedRef<SWidget> CreateMenu();
+
+	EVisibility HintText_GetVisibility() const;
 };
 
 } // UE::MediaViewer::Private
