Media Viewer: Removed mip value from user view (as it's not used).

diff --git a/Engine/Plugins/Experimental/MediaViewer/Source/MediaViewer/Public/ImageViewer/MediaImageViewer.h b/Engine/Plugins/Experimental/MediaViewer/Source/MediaViewer/Public/ImageViewer/MediaImageViewer.h
--- a/Engine/Plugins/Experimental/MediaViewer/Source/MediaViewer/Public/ImageViewer/MediaImageViewer.h
+++ b/Engine/Plugins/Experimental/MediaViewer/Source/MediaViewer/Public/ImageViewer/MediaImageViewer.h
@@ -91,8 +91,9 @@ struct FMediaImagePaintSettings
 
 	/**
 	 * The requested mip level. If the mip isn't available, the highest available value will be used.
+	 * Not currently used.
 	 */
-	UPROPERTY(EditAnywhere, Config, Category = "Media")
+	UPROPERTY(Config)
 	uint8 MipLevel = 0;
 
 	/**
