Media Viewer: Scale now shows only up to 1 decimal place.

diff --git a/Engine/Plugins/Experimental/MediaViewer/Source/MediaViewer/Private/Widgets/SMediaViewerToolbar.cpp b/Engine/Plugins/Experimental/MediaViewer/Source/MediaViewer/Private/Widgets/SMediaViewerToolbar.cpp
--- a/Engine/Plugins/Experimental/MediaViewer/Source/MediaViewer/Private/Widgets/SMediaViewerToolbar.cpp
+++ b/Engine/Plugins/Experimental/MediaViewer/Source/MediaViewer/Private/Widgets/SMediaViewerToolbar.cpp
@@ -338,7 +338,7 @@ FText SMediaViewerToolbar::GetScaleMenuLabel(EMediaImageViewerPosition InPositio
 	if (TSharedPtr<FMediaImageViewer> ImageViewer = Delegates->GetImageViewer.Execute(InPosition))
 	{
 		FNumberFormattingOptions FormattingFloat;
-		FormattingFloat.SetMinimumFractionalDigits(1);
+		FormattingFloat.SetMaximumFractionalDigits(1);
 
 		return FText::Format(
 			LOCTEXT("Scale", "{0}%"),
