Media Viewer: Drop target is now visible when the overlay is disabled.

diff --git a/Engine/Plugins/Experimental/MediaViewer/Source/MediaViewer/Private/Widgets/SMediaImageViewerOverlay.cpp b/Engine/Plugins/Experimental/MediaViewer/Source/MediaViewer/Private/Widgets/SMediaImageViewerOverlay.cpp
--- a/Engine/Plugins/Experimental/MediaViewer/Source/MediaViewer/Private/Widgets/SMediaImageViewerOverlay.cpp
+++ b/Engine/Plugins/Experimental/MediaViewer/Source/MediaViewer/Private/Widgets/SMediaImageViewerOverlay.cpp
@@ -29,6 +29,8 @@
 namespace UE::MediaViewer::Private
 {
 
+const FName DropTargetName = TEXT("DropTarget");
+
 SMediaImageViewerOverlay::SMediaImageViewerOverlay()
 	: CommandList(MakeShared<FUICommandList>())
 	, DragButtonName(NAME_None)
@@ -67,7 +69,7 @@ void SMediaImageViewerOverlay::Construct(const FArguments& InArgs, EMediaImageVi
 	[
 		SNew(SScissorRectBox)
 		[
-			SAssignNew(Container, SVerticalBox)
+			SNew(SVerticalBox)
 
 			+ SVerticalBox::Slot()
 			.AutoHeight()
@@ -433,7 +435,7 @@ TSharedRef<SWidget> SMediaImageViewerOverlay::CreateOverlay(bool bInComparisonVi
 		return SNullWidget::NullWidget;
 	}
 
-	TSharedRef<SOverlay> Overlay = SNew(SOverlay);
+	Overlay = SNew(SOverlay);
 
 	TSharedPtr<SWidget> ImageViewerOverlay = ImageViewer.IsValid() 
 		? ImageViewer->GetOverlayWidget(Position, Delegates->GetTab.Execute()) 
@@ -444,7 +446,7 @@ TSharedRef<SWidget> SMediaImageViewerOverlay::CreateOverlay(bool bInComparisonVi
 		.VAlign(VAlign_Center)
 		[
 			SNew(STextBlock)
-			.Visibility(this, &SMediaImageViewerOverlay::HintText_GetVisibility)
+			.Visibility(SMediaImageViewerOverlay::HintText_GetVisibility())
 			.Text(LOCTEXT("DropTargetMessage", "Drop supported asset or library item here."))
 			.Font(IDetailLayoutBuilder::GetDetailFont())
 			.TextStyle(FAppStyle::Get(), "HintText")
@@ -457,6 +459,7 @@ TSharedRef<SWidget> SMediaImageViewerOverlay::CreateOverlay(bool bInComparisonVi
 				SNew(SMediaViewerDropTarget, Delegates.ToSharedRef())
 				.Position(Position)
 				.bComparisonView(bInComparisonView)
+				.Tag(DropTargetName)
 			];
 	}
 	else
@@ -464,6 +467,7 @@ TSharedRef<SWidget> SMediaImageViewerOverlay::CreateOverlay(bool bInComparisonVi
 		Overlay->AddSlot()
 			[
 				SNew(SHorizontalBox)
+				.Tag(DropTargetName)
 				+ SHorizontalBox::Slot()
 				.FillWidth(0.25f)
 				.Padding(0.f, 0.f, 0.f, 20.f)
@@ -541,12 +545,12 @@ TSharedRef<SWidget> SMediaImageViewerOverlay::CreateOverlay(bool bInComparisonVi
 			]
 		];
 
-	return Overlay;
+	return Overlay.ToSharedRef();
 }
 
 TSharedRef<SWidget> SMediaImageViewerOverlay::CreateStatusBar(const TSharedRef<FMediaViewerDelegates>& InDelegates)
 {
-	return SNew(SMediaImageViewerStatusBar, Position, InDelegates);
+	return SAssignNew(StatusBar, SMediaImageViewerStatusBar, Position, InDelegates);
 }
 
 bool SMediaImageViewerOverlay::IsOverlayEnabled() const
@@ -557,18 +561,25 @@ bool SMediaImageViewerOverlay::IsOverlayEnabled() const
 void SMediaImageViewerOverlay::ToggleOverlay()
 {
 	bOverlayEnabled = !bOverlayEnabled;
+	const EVisibility NewVisibility = bOverlayEnabled ? EVisibility::Visible : EVisibility::Collapsed;
 
-	FChildren* Children = Container->GetChildren();
+	if (StatusBar.IsValid())
+	{
+		StatusBar->SetVisibility(NewVisibility);
+	}
 
-	for (int32 Index = 0; Index < Children->Num(); ++Index)
+	if (Overlay.IsValid())
 	{
-		if (bOverlayEnabled)
-		{
-			Children->GetChildAt(Index)->SetVisibility(EVisibility::Visible);
-		}
-		else
+		FChildren* Children = Overlay->GetChildren();
+
+		for (int32 Index = 0; Index < Children->Num(); ++Index)
 		{
-			Children->GetChildAt(Index)->SetVisibility(EVisibility::Collapsed);
+			TSharedRef<SWidget> Child = Children->GetChildAt(Index);
+
+			if (Child->GetTag() != DropTargetName)
+			{
+				Child->SetVisibility(NewVisibility);
+			}
 		}
 	}
 }
diff --git a/Engine/Plugins/Experimental/MediaViewer/Source/MediaViewer/Private/Widgets/SMediaImageViewerOverlay.h b/Engine/Plugins/Experimental/MediaViewer/Source/MediaViewer/Private/Widgets/SMediaImageViewerOverlay.h
--- a/Engine/Plugins/Experimental/MediaViewer/Source/MediaViewer/Private/Widgets/SMediaImageViewerOverlay.h
+++ b/Engine/Plugins/Experimental/MediaViewer/Source/MediaViewer/Private/Widgets/SMediaImageViewerOverlay.h
@@ -7,6 +7,7 @@
 #include "Widgets/MediaViewerDelegates.h"
 
 class FUICommandList;
+class SOverlay;
 class SVerticalBox;
 struct FMediaViewerLibraryItem;
 
@@ -19,6 +20,7 @@ namespace UE::MediaViewer
 namespace UE::MediaViewer::Private
 {
 
+class SMediaImageViewerStatusBar;
 class SMediaViewerLibrary;
 class SMediaViewerLibraryPrivate;
 
@@ -91,7 +93,8 @@ protected:
 	FVector DragStartOffset;
 	/** For moving the image based on cursor location. */
 	FIntPoint LastCursorPosition;
-	TSharedPtr<SVerticalBox> Container;
+	TSharedPtr<SMediaImageViewerStatusBar> StatusBar;
+	TSharedPtr<SOverlay> Overlay;
 	/** Toggles the visibility of the overlay. */
 	bool bOverlayEnabled;
 	/** Whether, on first valid paint, it should scale to fit. */
