MediaViewer: Do not load assets when they are right clicked in the content browser

diff --git a/Engine/Plugins/Experimental/MediaViewer/Source/MediaViewer/Private/ImageViewers/MaterialInterfaceImageViewer.cpp b/Engine/Plugins/Experimental/MediaViewer/Source/MediaViewer/Private/ImageViewers/MaterialInterfaceImageViewer.cpp
--- a/Engine/Plugins/Experimental/MediaViewer/Source/MediaViewer/Private/ImageViewers/MaterialInterfaceImageViewer.cpp
+++ b/Engine/Plugins/Experimental/MediaViewer/Source/MediaViewer/Private/ImageViewers/MaterialInterfaceImageViewer.cpp
@@ -3,11 +3,15 @@
 #include "ImageViewers/MaterialInterfaceImageViewer.h"
 
 #include "AssetRegistry/AssetData.h"
+#include "AssetRegistry/AssetRegistryModule.h"
+#include "AssetRegistry/IAssetRegistry.h"
 #include "Brushes/SlateImageBrush.h"
 #include "Engine/TextureRenderTarget2D.h"
 #include "Materials/Material.h"
 #include "Materials/MaterialInterface.h"
 #include "Materials/MaterialInstance.h"
+#include "Materials/MaterialInstanceConstant.h"
+#include "Materials/MaterialInstanceDynamic.h"
 #include "Math/Vector2D.h"
 #include "MediaViewerUtils.h"
 #include "Misc/Optional.h"
@@ -22,9 +26,25 @@ const FLazyName FMaterialInterfaceImageViewer::ItemTypeName = TEXT("MaterialInte
 
 bool FMaterialInterfaceImageViewer::FFactory::SupportsAsset(const FAssetData& InAssetData) const
 {
-	if (UObject* Object = InAssetData.GetAsset())
+	if (InAssetData.AssetClassPath == UMaterial::StaticClass()->GetClassPathName())
+	{
+		return InAssetData.GetTagValueRef<FString>(GET_MEMBER_NAME_CHECKED(UMaterial, MaterialDomain)) == TEXT("MD_UI");
+	}
+	else if (InAssetData.AssetClassPath == UMaterialInstanceConstant::StaticClass()->GetClassPathName() || InAssetData.AssetClassPath == UMaterialInstanceDynamic::StaticClass()->GetClassPathName())
 	{
-		return SupportsObject(TNotNull<UObject*>(Object));
+		// Find the parent material
+		FAssetRegistryModule& AssetRegistry = FModuleManager::Get().LoadModuleChecked<FAssetRegistryModule>("AssetRegistry");
+		FAssetData CurrentData = InAssetData;
+		while (CurrentData.IsValid() && CurrentData.AssetClassPath != UMaterial::StaticClass()->GetClassPathName())
+		{
+			FName Parent = InAssetData.GetTagValueRef<FName>(GET_MEMBER_NAME_CHECKED(UMaterialInstance, Parent));
+			CurrentData = AssetRegistry.Get().GetAssetByObjectPath(FSoftObjectPath(Parent.ToString()));
+		}
+
+		if (CurrentData.IsValid())
+		{
+			return CurrentData.GetTagValueRef<FString>(GET_MEMBER_NAME_CHECKED(UMaterial, MaterialDomain)) == TEXT("MD_UI");
+		}
 	}
 
 	return false;
