Media Viewer: Media track section
- Added media track section that lists tracks in sequencer that have a media texture.
- Added the refresh dynamic group button.

diff --git a/Engine/Plugins/Experimental/MediaViewer/Source/MediaViewer/MediaViewer.build.cs b/Engine/Plugins/Experimental/MediaViewer/Source/MediaViewer/MediaViewer.build.cs
--- a/Engine/Plugins/Experimental/MediaViewer/Source/MediaViewer/MediaViewer.build.cs
+++ b/Engine/Plugins/Experimental/MediaViewer/Source/MediaViewer/MediaViewer.build.cs
@@ -28,6 +28,7 @@ public class MediaViewer : ModuleRules
 				"InputCore",
 				"LevelEditor",
 				"MediaAssets",
+				"MediaCompositing",
 				"MediaPlayerEditor",
 				"MediaStream",
 				"Projects",
diff --git a/Engine/Plugins/Experimental/MediaViewer/Source/MediaViewer/Private/Library/MediaTrackGroup.cpp b/Engine/Plugins/Experimental/MediaViewer/Source/MediaViewer/Private/Library/MediaTrackGroup.cpp
new file mode 100644
--- /dev/null
+++ b/Engine/Plugins/Experimental/MediaViewer/Source/MediaViewer/Private/Library/MediaTrackGroup.cpp
@@ -0,0 +1,70 @@
+// Copyright Epic Games, Inc. All Rights Reserved.
+
+#include "Library/MediaTrackGroup.h"
+
+#include "ImageViewers/MediaTextureImageViewer.h"
+#include "Library/MediaViewerLibraryItem.h"
+#include "MediaPlayer.h"
+#include "MediaTexture.h"
+#include "MediaViewerUtils.h"
+#include "MovieSceneMediaSection.h"
+#include "MovieSceneMediaTrack.h"
+#include "UObject/UObjectIterator.h"
+
+#define LOCTEXT_NAMESPACE "MediaTrackGroup"
+
+namespace UE::MediaViewer::Private
+{
+
+	FMediaTrackGroup::FMediaTrackGroup(const TSharedRef<IMediaViewerLibrary>& InLibrary)
+		: FMediaTrackGroup(InLibrary, FGuid::NewGuid())
+	{
+	}
+
+	FMediaTrackGroup::FMediaTrackGroup(const TSharedRef<IMediaViewerLibrary>& InLibrary,
+		const FGuid& InGuid)
+		: FMediaViewerLibraryDynamicGroup(
+			InLibrary,
+			InGuid,
+			LOCTEXT("MediaTracks", "Media Tracks"),
+			LOCTEXT("MediaTracksTooltip", "The media textures from media tracks."),
+			FGenerateItems::CreateStatic(&FMediaTrackGroup::GetMediaTrackItems)
+		)
+	{
+	}
+
+	TArray<TSharedRef<FMediaViewerLibraryItem>> FMediaTrackGroup::GetMediaTrackItems()
+	{
+		TArray<TSharedRef<FMediaViewerLibraryItem>> MediaTrackItems;
+
+		TSharedRef<FMediaTextureImageViewer::FFactory> Factory = MakeShared<FMediaTextureImageViewer::FFactory>();
+
+		for (UMovieSceneMediaSection* MediaSection : TObjectRange<UMovieSceneMediaSection>())
+		{
+			if (MediaSection->IsTemplate() || !MediaSection->MediaTexture)
+			{
+				continue;
+			}
+
+			if (TSharedPtr<FMediaViewerLibraryItem> NewItem = Factory->CreateLibraryItem(TNotNull<UMediaTexture*>(MediaSection->MediaTexture)))
+			{
+				if (UMovieSceneMediaTrack* MediaTrack = Cast<UMovieSceneMediaTrack>(MediaSection->GetOuter()))
+				{
+					const FText DisplayName = MediaTrack->GetDisplayName();
+
+					if (!DisplayName.IsEmpty())
+					{
+						NewItem->Name = DisplayName;
+					}
+				}
+
+				MediaTrackItems.Add(NewItem.ToSharedRef());
+			}
+		}
+
+		return MediaTrackItems;
+	}
+
+} // UE::MediaViewer::Private
+
+#undef LOCTEXT_NAMESPACE
diff --git a/Engine/Plugins/Experimental/MediaViewer/Source/MediaViewer/Private/Library/MediaTrackGroup.h b/Engine/Plugins/Experimental/MediaViewer/Source/MediaViewer/Private/Library/MediaTrackGroup.h
new file mode 100644
--- /dev/null
+++ b/Engine/Plugins/Experimental/MediaViewer/Source/MediaViewer/Private/Library/MediaTrackGroup.h
@@ -0,0 +1,24 @@
+// Copyright Epic Games, Inc. All Rights Reserved.
+
+#pragma once
+
+#include "Library/MediaViewerLibraryDynamicGroup.h"
+
+namespace UE::MediaViewer::Private
+{
+
+	/**
+	 * A group that generates entries based on available UMediaTracks.
+	*/
+	struct FMediaTrackGroup : public FMediaViewerLibraryDynamicGroup
+	{
+		FMediaTrackGroup(const TSharedRef<IMediaViewerLibrary>& InLibrary);
+		FMediaTrackGroup(const TSharedRef<IMediaViewerLibrary>& InLibrary, const FGuid& InId);
+
+		virtual ~FMediaTrackGroup() override = default;
+
+	protected:
+		static TArray<TSharedRef<FMediaViewerLibraryItem>> GetMediaTrackItems();
+	};
+
+} // UE::MediaViewer::Private
diff --git a/Engine/Plugins/Experimental/MediaViewer/Source/MediaViewer/Private/Library/MediaViewerLibrary.cpp b/Engine/Plugins/Experimental/MediaViewer/Source/MediaViewer/Private/Library/MediaViewerLibrary.cpp
--- a/Engine/Plugins/Experimental/MediaViewer/Source/MediaViewer/Private/Library/MediaViewerLibrary.cpp
+++ b/Engine/Plugins/Experimental/MediaViewer/Source/MediaViewer/Private/Library/MediaViewerLibrary.cpp
@@ -3,6 +3,7 @@
 #include "Library/MediaViewerLibrary.h"
 
 #include "Containers/UnrealString.h"
+#include "Library/MediaViewerLibraryDynamicGroup.h"
 #include "Library/MediaViewerLibraryGroup.h"
 #include "Library/MediaViewerLibraryItem.h"
 #include "Misc/Guid.h"
@@ -82,7 +83,16 @@ bool FMediaViewerLibrary::AddGroup(const TSharedRef<FMediaViewerLibraryGroup>& I
 		return false;
 	}
 
-	Groups.Add(InNewGroup);;
+	Groups.Add(InNewGroup);
+
+	if (InNewGroup->IsDynamic())
+	{
+		StaticCastSharedRef<FMediaViewerLibraryDynamicGroup>(InNewGroup)->GetOnItemsUpdated().AddSP(
+			this,
+			&FMediaViewerLibrary::OnChanged,
+			EChangeType::GroupItemsChanged
+		);
+	}
 
 	OnChanged(IMediaViewerLibrary::EChangeType::GroupAdded);
 
@@ -139,6 +149,11 @@ TSharedPtr<FMediaViewerLibraryGroup> FMediaViewerLibrary::RemoveGroup(const FGui
 
 	TSharedPtr<FMediaViewerLibraryGroup> GroupToRemove = Groups[GroupToRemoveIndex];
 
+	if (GroupToRemove->IsDynamic())
+	{
+		StaticCastSharedPtr<FMediaViewerLibraryDynamicGroup>(GroupToRemove)->GetOnItemsUpdated().RemoveAll(this);
+	}
+
 	Groups.RemoveAt(GroupToRemoveIndex);
 
 	OnChanged(IMediaViewerLibrary::EChangeType::GroupRemoved);
diff --git a/Engine/Plugins/Experimental/MediaViewer/Source/MediaViewer/Private/Library/MediaViewerLibraryDynamicGroup.cpp b/Engine/Plugins/Experimental/MediaViewer/Source/MediaViewer/Private/Library/MediaViewerLibraryDynamicGroup.cpp
--- a/Engine/Plugins/Experimental/MediaViewer/Source/MediaViewer/Private/Library/MediaViewerLibraryDynamicGroup.cpp
+++ b/Engine/Plugins/Experimental/MediaViewer/Source/MediaViewer/Private/Library/MediaViewerLibraryDynamicGroup.cpp
@@ -27,17 +27,31 @@ FMediaViewerLibraryDynamicGroup::FMediaViewerLibraryDynamicGroup(const TSharedRe
 
 void FMediaViewerLibraryDynamicGroup::UpdateItems()
 {
-	Items = GetUpdatedIs(Items);
+	TOptional<TArray<FGuid>> NewItems = GetUpdatedIs(Items);
+
+	if (!NewItems.IsSet())
+	{
+		return;
+	}
+
+	Items = NewItems.GetValue();
+
+	OnItemsUpdated.Broadcast();
 }
 
-TArray<FGuid> FMediaViewerLibraryDynamicGroup::GetUpdatedIs(const TArray<FGuid>& InCurrentIds)
+FSimpleMulticastDelegate::RegistrationType& FMediaViewerLibraryDynamicGroup::GetOnItemsUpdated()
+{
+	return OnItemsUpdated;
+}
+
+TOptional<TArray<FGuid>> FMediaViewerLibraryDynamicGroup::GetUpdatedIs(const TArray<FGuid>& InCurrentIds)
 {
 	// Check if we still have a valid Library
 	TSharedPtr<IMediaViewerLibrary> Library = LibraryWeak.Pin();
 
 	if (!Library.IsValid())
 	{
-		return InCurrentIds;
+		return {};
 	}
 
 	// Populate map of current items
@@ -53,6 +67,7 @@ TArray<FGuid> FMediaViewerLibraryDynamicGroup::GetUpdatedIs(const TArray<FGuid>&
 
 	// Generate new id list
 	TArray<FGuid> UpdatedIds;
+	bool bHasNewItem = false;
 
 	const TArray<TSharedRef<FMediaViewerLibraryItem>> UpdatedItems = GenerateItemsDelegate.Execute();
 
@@ -69,9 +84,16 @@ TArray<FGuid> FMediaViewerLibraryDynamicGroup::GetUpdatedIs(const TArray<FGuid>&
 		{
 			Library->AddItem(UpdatedItem);
 			UpdatedIds.Add(UpdatedItem->GetId());
+			bHasNewItem = true;
 		}
 	}
 
+	// No new items and no deleted items
+	if (!bHasNewItem && ExistingItems.IsEmpty())
+	{
+		return {};
+	}
+
 	// Remove invalidated items
 	for (const TPair<FString, FGuid>& ItemPair : ExistingItems)
 	{
diff --git a/Engine/Plugins/Experimental/MediaViewer/Source/MediaViewer/Private/MediaViewerStyle.cpp b/Engine/Plugins/Experimental/MediaViewer/Source/MediaViewer/Private/MediaViewerStyle.cpp
--- a/Engine/Plugins/Experimental/MediaViewer/Source/MediaViewer/Private/MediaViewerStyle.cpp
+++ b/Engine/Plugins/Experimental/MediaViewer/Source/MediaViewer/Private/MediaViewerStyle.cpp
@@ -52,6 +52,14 @@ FMediaViewerStyle::FMediaViewerStyle()
 		.SetNormalPadding(FMargin(2.f, 2.f, 2.f, 2.f))
 		.SetPressedPadding(FMargin(2.f, 2.f, 2.f, 2.f)));
 
+	Set("LibraryButtons", FButtonStyle(FAppStyle::Get().GetWidgetStyle<FButtonStyle>("Animation.PlayControlsButton"))
+		.SetNormal(FSlateNoResource())
+		.SetDisabled(FSlateNoResource())
+		.SetHovered(FSlateRoundedBoxBrush(FLinearColor(0.2, 0.2, 0.2, 0.5), 3.f, FVector2D(20.f)))
+		.SetPressed(FSlateRoundedBoxBrush(FLinearColor(0.1, 0.1, 0.1, 0.5), 3.f, FVector2D(20.f)))
+		.SetNormalPadding(FMargin(2.f, 0.f, 2.f, 0.f))
+		.SetPressedPadding(FMargin(2.f, 0.f, 2.f, 0.f)));
+
 	FSlateStyleRegistry::RegisterSlateStyle(*this);
 }
 
diff --git a/Engine/Plugins/Experimental/MediaViewer/Source/MediaViewer/Private/Widgets/SMediaViewerLibrary.cpp b/Engine/Plugins/Experimental/MediaViewer/Source/MediaViewer/Private/Widgets/SMediaViewerLibrary.cpp
--- a/Engine/Plugins/Experimental/MediaViewer/Source/MediaViewer/Private/Widgets/SMediaViewerLibrary.cpp
+++ b/Engine/Plugins/Experimental/MediaViewer/Source/MediaViewer/Private/Widgets/SMediaViewerLibrary.cpp
@@ -162,9 +162,7 @@ void SMediaViewerLibrary::OnGroupExpanded(IMediaViewerLibrary::FGroupItem InGrou
 	if (Group.IsValid() && Group->IsDynamic())
 	{
 		StaticCastSharedPtr<FMediaViewerLibraryDynamicGroup>(Group)->UpdateItems();
-	}	
-
-	TreeView->RequestListRefresh();
+	}
 }
 
 } // UE::MediaViewer::Private
diff --git a/Engine/Plugins/Experimental/MediaViewer/Source/MediaViewer/Private/Widgets/SMediaViewerLibraryGroup.cpp b/Engine/Plugins/Experimental/MediaViewer/Source/MediaViewer/Private/Widgets/SMediaViewerLibraryGroup.cpp
--- a/Engine/Plugins/Experimental/MediaViewer/Source/MediaViewer/Private/Widgets/SMediaViewerLibraryGroup.cpp
+++ b/Engine/Plugins/Experimental/MediaViewer/Source/MediaViewer/Private/Widgets/SMediaViewerLibraryGroup.cpp
@@ -8,6 +8,7 @@
 #include "DragAndDrop/AssetDragDropOp.h"
 #include "IMediaViewerModule.h"
 #include "Library/MediaViewerLibrary.h"
+#include "Library/MediaViewerLibraryDynamicGroup.h"
 #include "Library/MediaViewerLibraryGroup.h"
 #include "Library/MediaViewerLibraryItem.h"
 #include "MediaViewerStyle.h"
@@ -57,6 +58,25 @@ void SMediaViewerLibraryGroup::Construct(const FArguments& InArgs, const TShared
 
 	if (Group->IsDynamic())
 	{
+		const FButtonStyle& ButtonStyle = FMediaViewerStyle::Get().GetWidgetStyle<FButtonStyle>("LibraryButtons");
+
+		Inner->AddSlot()
+			.AutoWidth()
+			.Padding(0.f, 4.f, 5.f, 4.f)
+			.VAlign(VAlign_Center)
+			[
+				SNew(SButton)
+				.ContentPadding(FMargin(1.f, 3.f))
+				.OnClicked(this, &SMediaViewerLibraryGroup::OnRefreshButtonClicked)
+				.ButtonStyle(&ButtonStyle)
+				[
+					SNew(SImage)
+					.Image(FAppStyle::Get().GetBrush("SourceControl.Actions.Refresh"))
+					.ToolTipText(LOCTEXT("RefreshGroupToolTip", "Refresh the members of this dynamic group."))
+					.DesiredSizeOverride(FVector2D(12))
+				]
+			];
+
 		Inner->AddSlot()
 			.AutoWidth()
 			.Padding(0.f, 4.f, 5.f, 4.f)
@@ -64,13 +84,13 @@ void SMediaViewerLibraryGroup::Construct(const FArguments& InArgs, const TShared
 			[
 				SNew(SImage)
 				.Image(FAppStyle::Get().GetBrush("Icons.Lock"))
-				.ToolTipText(LOCTEXT("DynamicGroup", "This group is generated dynamically. Items cannot be added or removed."))
+				.ToolTipText(LOCTEXT("DynamicGroupToolTip", "This group is generated dynamically. Items cannot be added or removed."))
 				.DesiredSizeOverride(FVector2D(12))
 			];
 	}
 	else if (InLibrary->CanRemoveGroup(Group->GetId()))
 	{
-		const FButtonStyle& ButtonStyle = FAppStyle::Get().GetWidgetStyle<FButtonStyle>("NoBorder");
+		const FButtonStyle& ButtonStyle = FMediaViewerStyle::Get().GetWidgetStyle<FButtonStyle>("LibraryButtons");
 
 		Inner->AddSlot()
 			.AutoWidth()
@@ -216,6 +236,27 @@ FReply SMediaViewerLibraryGroup::OnRemoveButtonClicked()
 	return FReply::Handled();
 }
 
+FReply SMediaViewerLibraryGroup::OnRefreshButtonClicked()
+{
+	TSharedPtr<FMediaViewerLibrary> Library = LibraryWeak.Pin();
+
+	if (!Library.IsValid())
+	{
+		return FReply::Handled();
+	}
+
+	TSharedPtr<FMediaViewerLibraryGroup> Group = Library->GetGroup(GroupId);
+
+	if (!Group.IsValid() || !Group->IsDynamic())
+	{
+		return FReply::Handled();
+	}
+
+	StaticCastSharedPtr<FMediaViewerLibraryDynamicGroup>(Group)->UpdateItems();
+
+	return FReply::Handled();
+}
+
 bool SMediaViewerLibraryGroup::CanAcceptLibraryItem(const IMediaViewerLibrary::FGroupItem& InDraggedGroupItem) const
 {
 	TSharedPtr<FMediaViewerLibrary> Library = LibraryWeak.Pin();
diff --git a/Engine/Plugins/Experimental/MediaViewer/Source/MediaViewer/Private/Widgets/SMediaViewerLibraryGroup.h b/Engine/Plugins/Experimental/MediaViewer/Source/MediaViewer/Private/Widgets/SMediaViewerLibraryGroup.h
--- a/Engine/Plugins/Experimental/MediaViewer/Source/MediaViewer/Private/Widgets/SMediaViewerLibraryGroup.h
+++ b/Engine/Plugins/Experimental/MediaViewer/Source/MediaViewer/Private/Widgets/SMediaViewerLibraryGroup.h
@@ -53,6 +53,8 @@ protected:
 
 	FReply OnRemoveButtonClicked();
 
+	FReply OnRefreshButtonClicked();
+
 	bool CanAcceptLibraryItem(const IMediaViewerLibrary::FGroupItem& InDraggedGroupItem) const;
 	void OnLibraryItemDropped(const IMediaViewerLibrary::FGroupItem& InDroppedGroupItem) const;
 
diff --git a/Engine/Plugins/Experimental/MediaViewer/Source/MediaViewer/Private/Widgets/SMediaViewerLibraryPrivate.cpp b/Engine/Plugins/Experimental/MediaViewer/Source/MediaViewer/Private/Widgets/SMediaViewerLibraryPrivate.cpp
--- a/Engine/Plugins/Experimental/MediaViewer/Source/MediaViewer/Private/Widgets/SMediaViewerLibraryPrivate.cpp
+++ b/Engine/Plugins/Experimental/MediaViewer/Source/MediaViewer/Private/Widgets/SMediaViewerLibraryPrivate.cpp
@@ -6,6 +6,7 @@
 #include "LevelEditor.h"
 #include "Library/LevelEditorViewportGroup.h"
 #include "Library/MediaTextureGroup.h"
+#include "Library/MediaTrackGroup.h"
 #include "Library/MediaViewerLibrary.h"
 #include "Library/MediaViewerLibraryGroup.h"
 #include "Library/MediaViewerLibraryIni.h"
@@ -42,6 +43,7 @@ void SMediaViewerLibraryPrivate::Construct(const FArguments& InArgs, const TShar
 	TSharedRef<FMediaViewerLibrary> LibraryImpl = GetLibrary();
 
 	LibraryImpl->AddGroup(MakeShared<FLevelEditorViewportGroup>(LibraryImpl));
+	LibraryImpl->AddGroup(MakeShared<FMediaTrackGroup>(LibraryImpl));
 	LibraryImpl->AddGroup(MakeShared<FMediaTextureGroup>(LibraryImpl));
 }
 
diff --git a/Engine/Plugins/Experimental/MediaViewer/Source/MediaViewer/Public/Library/IMediaViewerLibrary.h b/Engine/Plugins/Experimental/MediaViewer/Source/MediaViewer/Public/Library/IMediaViewerLibrary.h
--- a/Engine/Plugins/Experimental/MediaViewer/Source/MediaViewer/Public/Library/IMediaViewerLibrary.h
+++ b/Engine/Plugins/Experimental/MediaViewer/Source/MediaViewer/Public/Library/IMediaViewerLibrary.h
@@ -48,6 +48,7 @@ public:
 	{
 		GroupAdded,
 		GroupRemoved,
+		GroupItemsChanged,
 		ItemAdded,
 		ItemRemoved,
 		ItemGroupChanged
diff --git a/Engine/Plugins/Experimental/MediaViewer/Source/MediaViewer/Public/Library/MediaViewerLibraryDynamicGroup.h b/Engine/Plugins/Experimental/MediaViewer/Source/MediaViewer/Public/Library/MediaViewerLibraryDynamicGroup.h
--- a/Engine/Plugins/Experimental/MediaViewer/Source/MediaViewer/Public/Library/MediaViewerLibraryDynamicGroup.h
+++ b/Engine/Plugins/Experimental/MediaViewer/Source/MediaViewer/Public/Library/MediaViewerLibraryDynamicGroup.h
@@ -4,6 +4,9 @@
 
 #include "Library/MediaViewerLibraryGroup.h"
 
+#include "Delegates/Delegate.h"
+#include "Misc/Optional.h"
+
 namespace UE::MediaViewer
 {
 
@@ -26,11 +29,14 @@ struct FMediaViewerLibraryDynamicGroup : public FMediaViewerLibraryGroup
 
 	MEDIAVIEWER_API virtual void UpdateItems();
 
+	MEDIAVIEWER_API FSimpleMulticastDelegate::RegistrationType& GetOnItemsUpdated();
+
 protected:
 	TWeakPtr<IMediaViewerLibrary> LibraryWeak;
 	FGenerateItems GenerateItemsDelegate;
+	FSimpleMulticastDelegate OnItemsUpdated;
 
-	TArray<FGuid> GetUpdatedIs(const TArray<FGuid>& InCurrentIds);
+	TOptional<TArray<FGuid>> GetUpdatedIs(const TArray<FGuid>& InCurrentIds);
 };
 
 } // UE::MediaViewer
