Media Viewer: Right click asset fix
- There's now a maximum parent check for material instances

diff --git a/Engine/Plugins/Experimental/MediaViewer/Source/MediaViewer/Private/ImageViewers/MaterialInterfaceImageViewer.cpp b/Engine/Plugins/Experimental/MediaViewer/Source/MediaViewer/Private/ImageViewers/MaterialInterfaceImageViewer.cpp
--- a/Engine/Plugins/Experimental/MediaViewer/Source/MediaViewer/Private/ImageViewers/MaterialInterfaceImageViewer.cpp
+++ b/Engine/Plugins/Experimental/MediaViewer/Source/MediaViewer/Private/ImageViewers/MaterialInterfaceImageViewer.cpp
@@ -35,10 +35,21 @@ bool FMaterialInterfaceImageViewer::FFactory::SupportsAsset(const FAssetData& In
 		// Find the parent material
 		FAssetRegistryModule& AssetRegistry = FModuleManager::Get().LoadModuleChecked<FAssetRegistryModule>("AssetRegistry");
 		FAssetData CurrentData = InAssetData;
+		int32 ParentCount = 0;
+		constexpr int32 MaxParentCheck = 20;
+
 		while (CurrentData.IsValid() && CurrentData.AssetClassPath != UMaterial::StaticClass()->GetClassPathName())
 		{
 			FName Parent = InAssetData.GetTagValueRef<FName>(GET_MEMBER_NAME_CHECKED(UMaterialInstance, Parent));
 			CurrentData = AssetRegistry.Get().GetAssetByObjectPath(FSoftObjectPath(Parent.ToString()));
+
+			++ParentCount;
+
+			if (ParentCount == MaxParentCheck)
+			{
+				UE_LOG(LogMediaViewer, Error, TEXT("Unable to resolve material parent for \"{0}\"."), *InAssetData.GetObjectPathString());
+				break;
+			}
 		}
 
 		if (CurrentData.IsValid())
