Media Viewer: Fixed struct initial values

diff --git a/Engine/Plugins/Experimental/MediaViewer/Source/MediaViewer/Private/Library/MediaViewerLibraryIni.h b/Engine/Plugins/Experimental/MediaViewer/Source/MediaViewer/Private/Library/MediaViewerLibraryIni.h
--- a/Engine/Plugins/Experimental/MediaViewer/Source/MediaViewer/Private/Library/MediaViewerLibraryIni.h
+++ b/Engine/Plugins/Experimental/MediaViewer/Source/MediaViewer/Private/Library/MediaViewerLibraryIni.h
@@ -9,11 +9,10 @@
 #include "ImageViewer/MediaImageViewer.h"
 #include "Misc/Guid.h"
 #include "Widgets/MediaViewerSettings.h"
+#include "Widgets/SMediaViewer.h"
 
 #include "MediaViewerLibraryIni.generated.h"
 
-enum class EMediaImageViewerActivePosition : uint8;
-
 namespace UE::MediaViewer::Private
 {
 	class FMediaViewerLibrary;
@@ -58,7 +57,7 @@ struct FMediaViewerState
 	FMediaViewerSettings ViewerSettings;
 
 	UPROPERTY(Config)
-	EMediaImageViewerActivePosition ActiveView;
+	EMediaImageViewerActivePosition ActiveView = EMediaImageViewerActivePosition::Single;
 
 	UPROPERTY(Config)
 	TArray<FMediaViewerImageState> Images;
diff --git a/Engine/Plugins/Experimental/MediaViewer/Source/MediaViewer/Public/Library/MediaViewerLibraryEntry.h b/Engine/Plugins/Experimental/MediaViewer/Source/MediaViewer/Public/Library/MediaViewerLibraryEntry.h
--- a/Engine/Plugins/Experimental/MediaViewer/Source/MediaViewer/Public/Library/MediaViewerLibraryEntry.h
+++ b/Engine/Plugins/Experimental/MediaViewer/Source/MediaViewer/Public/Library/MediaViewerLibraryEntry.h
@@ -56,7 +56,7 @@ struct FMediaViewerLibraryEntry
 
 protected:
 	/** Unique id for this entry. */
-	UPROPERTY()
+	UPROPERTY(meta = (IgnoreForMemberInitializationTest))
 	FGuid Id;
 };
 
