Media Viewer: Crash with invalid library thumbnail

diff --git a/Engine/Plugins/Experimental/MediaViewer/Source/MediaViewer/Private/Widgets/SMediaViewerLibraryItem.cpp b/Engine/Plugins/Experimental/MediaViewer/Source/MediaViewer/Private/Widgets/SMediaViewerLibraryItem.cpp
--- a/Engine/Plugins/Experimental/MediaViewer/Source/MediaViewer/Private/Widgets/SMediaViewerLibraryItem.cpp
+++ b/Engine/Plugins/Experimental/MediaViewer/Source/MediaViewer/Private/Widgets/SMediaViewerLibraryItem.cpp
@@ -74,7 +74,7 @@ void SMediaViewerLibraryItem::Construct(const FArguments& InArgs, const TSharedR
 		.VAlign(VAlign_Center)
 		[
 			SNew(SImage)
-			.Image(ThumbnailBrush.IsValid() ? ThumbnailBrush.Get() : nullptr)
+			.Image(this, &SMediaViewerLibraryItem::GetThumbnail)
 			.DesiredSizeOverride(FVector2D(24.f))
 		] 
 		+ SHorizontalBox::Slot()
@@ -533,6 +533,16 @@ void SMediaViewerLibraryItem::OnAssetsDropped(TConstArrayView<FAssetData> InAsse
 	}
 }
 
+const FSlateBrush* SMediaViewerLibraryItem::GetThumbnail() const
+{
+	if (ThumbnailBrush.IsValid() && IsValid(ThumbnailBrush->GetResourceObject()))
+	{
+		return ThumbnailBrush.Get();
+	}
+
+	return FAppStyle::Get().GetBrush("SourceControl.StatusIcon.Unknown");
+}
+
 } // UE::MediaViewer::Private
 
 #undef LOCTEXT_NAMESPACE
diff --git a/Engine/Plugins/Experimental/MediaViewer/Source/MediaViewer/Private/Widgets/SMediaViewerLibraryItem.h b/Engine/Plugins/Experimental/MediaViewer/Source/MediaViewer/Private/Widgets/SMediaViewerLibraryItem.h
--- a/Engine/Plugins/Experimental/MediaViewer/Source/MediaViewer/Private/Widgets/SMediaViewerLibraryItem.h
+++ b/Engine/Plugins/Experimental/MediaViewer/Source/MediaViewer/Private/Widgets/SMediaViewerLibraryItem.h
@@ -79,6 +79,8 @@ protected:
 
 	bool CanAcceptAssets(TConstArrayView<FAssetData> InAssetData) const;
 	void OnAssetsDropped(TConstArrayView<FAssetData> InAssetData);
+
+	const FSlateBrush* GetThumbnail() const;
 };
 
 } // UE::MediaViewer::Private
