Media Viewer: Fixed crash when 2 viewport tabs are open.

diff --git a/Engine/Plugins/Experimental/MediaViewer/Source/MediaViewer/Private/Library/LevelEditorViewportGroup.cpp b/Engine/Plugins/Experimental/MediaViewer/Source/MediaViewer/Private/Library/LevelEditorViewportGroup.cpp
--- a/Engine/Plugins/Experimental/MediaViewer/Source/MediaViewer/Private/Library/LevelEditorViewportGroup.cpp
+++ b/Engine/Plugins/Experimental/MediaViewer/Source/MediaViewer/Private/Library/LevelEditorViewportGroup.cpp
@@ -45,6 +45,9 @@ TArray<TSharedRef<FMediaViewerLibraryItem>> FLevelEditorViewportGroup::GetLevelE
 	}
 
 	TArray<TSharedPtr<SLevelViewport>> Viewports = LevelEditor->GetViewports();
+	
+	TSet<FGuid> AddedViewportIds;
+	AddedViewportIds.Reserve(Viewports.Num());
 
 	for (int32 Index = 0; Index < Viewports.Num(); ++Index)
 	{
@@ -63,11 +66,20 @@ TArray<TSharedRef<FMediaViewerLibraryItem>> FLevelEditorViewportGroup::GetLevelE
 		}
 
 		const FString ConfigKey = Viewports[Index]->GetConfigKey().ToString();
+
 		const FGuid ViewportId = FLevelEditorViewportImageViewer::FItem::GetIdForViewport(
 			ConfigKey,
 			/* Create id if invalid */ false
 		);
 
+		bool bAlreadyInSet = false;
+		AddedViewportIds.Add(ViewportId, &bAlreadyInSet);
+
+		if (bAlreadyInSet)
+		{
+			continue;
+		}
+
 		if (ViewportId.IsValid())
 		{
 			LevelEditorViewportItems.Add(MakeShared<FLevelEditorViewportImageViewer::FItem>(ViewportId, ConfigKey));
