Media Viewer: Improved level editor viewport robustness
- Made sure all settings have default values
- Moved level viewport render to end of RT thread event.
- Added cross-thread robustness to the pixel color system.

diff --git a/Engine/Plugins/Experimental/MediaViewer/Source/MediaViewer/Private/ImageViewers/LevelEditorViewportImageViewer.cpp b/Engine/Plugins/Experimental/MediaViewer/Source/MediaViewer/Private/ImageViewers/LevelEditorViewportImageViewer.cpp
--- a/Engine/Plugins/Experimental/MediaViewer/Source/MediaViewer/Private/ImageViewers/LevelEditorViewportImageViewer.cpp
+++ b/Engine/Plugins/Experimental/MediaViewer/Source/MediaViewer/Private/ImageViewers/LevelEditorViewportImageViewer.cpp
@@ -159,7 +159,12 @@ TSharedPtr<FSlateBrush> FLevelEditorViewportImageViewer::FItem::CreateThumbnail(
 	const FIntPoint RenderTargetSize = {32, 32};
 	RenderTarget = FMediaViewerUtils::CreateRenderTarget(RenderTargetSize, /* Transparent */ false);
 
-	RenderViewport(TNotNull<FViewport*>(Viewport.Get()), TNotNull<UTextureRenderTarget2D*>(RenderTarget), /* Resize render target */ false);
+	RenderViewport(
+		TNotNull<FViewport*>(Viewport.Get()), 
+		TNotNull<UTextureRenderTarget2D*>(RenderTarget), 
+		FRenderComplete(),
+		/* Resize render target */ false
+	);
 
 	TSharedRef<FSlateImageBrush> ThumbnailBrush = MakeShared<FSlateImageBrush>(RenderTarget, FVector2D(RenderTarget->GetSurfaceWidth(), RenderTarget->GetSurfaceDepth()));
 
@@ -223,7 +228,7 @@ TSharedPtr<FMediaViewerLibraryItem> FLevelEditorViewportImageViewer::CreateLibra
 
 void FLevelEditorViewportImageViewer::PaintImage(FMediaImagePaintParams& InPaintParams, const FMediaImagePaintGeometry& InPaintGeometry)
 {
-	if (ViweportSettings.bRealTime && !ViewportWeak.IsValid())
+	if (ViewportSettings.bRealTime && !ViewportWeak.IsValid())
 	{
 		if (StringValue.IsEmpty())
 		{
diff --git a/Engine/Plugins/Experimental/MediaViewer/Source/MediaViewer/Private/ImageViewers/MaterialInterfaceImageViewer.h b/Engine/Plugins/Experimental/MediaViewer/Source/MediaViewer/Private/ImageViewers/MaterialInterfaceImageViewer.h
--- a/Engine/Plugins/Experimental/MediaViewer/Source/MediaViewer/Private/ImageViewers/MaterialInterfaceImageViewer.h
+++ b/Engine/Plugins/Experimental/MediaViewer/Source/MediaViewer/Private/ImageViewers/MaterialInterfaceImageViewer.h
@@ -22,14 +22,14 @@ struct FMaterialInterfaceImageViewerSettings
 	GENERATED_BODY()
 
 	UPROPERTY(VisibleAnywhere, Category = "Material")
-	TObjectPtr<UMaterialInterface> MaterialInterface;
+	TObjectPtr<UMaterialInterface> MaterialInterface = nullptr;
 
 	UPROPERTY(VisibleAnywhere, Category = "Material")
-	TObjectPtr<UTextureRenderTarget2D> RenderTarget;
+	TObjectPtr<UTextureRenderTarget2D> RenderTarget = nullptr;
 
 	/** Will render the material every frame. */
 	UPROPERTY(EditAnywhere, Category = "Material")
-	bool bRealTime;
+	bool bRealTime = false;
 
 	UPROPERTY(EditAnywhere, Category = "Material", meta = (ClampMin = 32, UIMin = 32, ClampMax = 2048, UIMax = 2048))
 	int32 RenderTargetSize = 256;
diff --git a/Engine/Plugins/Experimental/MediaViewer/Source/MediaViewer/Private/ImageViewers/MediaSourceImageViewer.h b/Engine/Plugins/Experimental/MediaViewer/Source/MediaViewer/Private/ImageViewers/MediaSourceImageViewer.h
--- a/Engine/Plugins/Experimental/MediaViewer/Source/MediaViewer/Private/ImageViewers/MediaSourceImageViewer.h
+++ b/Engine/Plugins/Experimental/MediaViewer/Source/MediaViewer/Private/ImageViewers/MediaSourceImageViewer.h
@@ -25,7 +25,7 @@ struct FMediaSourceImageViewerSettings
 	GENERATED_BODY()
 
 	UPROPERTY(VisibleAnywhere, Category = "Texture")
-	TObjectPtr<UMediaSource> MediaSource;
+	TObjectPtr<UMediaSource> MediaSource = nullptr;
 };
 
 namespace UE::MediaViewer::Private
diff --git a/Engine/Plugins/Experimental/MediaViewer/Source/MediaViewer/Private/ImageViewers/MediaTextureImageViewer.h b/Engine/Plugins/Experimental/MediaViewer/Source/MediaViewer/Private/ImageViewers/MediaTextureImageViewer.h
--- a/Engine/Plugins/Experimental/MediaViewer/Source/MediaViewer/Private/ImageViewers/MediaTextureImageViewer.h
+++ b/Engine/Plugins/Experimental/MediaViewer/Source/MediaViewer/Private/ImageViewers/MediaTextureImageViewer.h
@@ -23,7 +23,7 @@ struct FMediaTextureImageViewerSettings
 	bool bRealTime = true;
 
 	UPROPERTY(VisibleAnywhere, Category = "Texture")
-	TObjectPtr<UMediaTexture> MediaTexture;
+	TObjectPtr<UMediaTexture> MediaTexture = nullptr;
 };
 
 namespace UE::MediaViewer::Private
diff --git a/Engine/Plugins/Experimental/MediaViewer/Source/MediaViewer/Private/ImageViewers/Texture2DImageViewer.h b/Engine/Plugins/Experimental/MediaViewer/Source/MediaViewer/Private/ImageViewers/Texture2DImageViewer.h
--- a/Engine/Plugins/Experimental/MediaViewer/Source/MediaViewer/Private/ImageViewers/Texture2DImageViewer.h
+++ b/Engine/Plugins/Experimental/MediaViewer/Source/MediaViewer/Private/ImageViewers/Texture2DImageViewer.h
@@ -22,7 +22,7 @@ struct FTexture2DImageViewerSettings
 	GENERATED_BODY()
 
 	UPROPERTY(VisibleAnywhere, Category = "Texture")
-	TObjectPtr<UTexture2D> Texture;
+	TObjectPtr<UTexture2D> Texture = nullptr;
 };
 
 namespace UE::MediaViewer::Private
diff --git a/Engine/Plugins/Experimental/MediaViewer/Source/MediaViewer/Private/ImageViewers/TextureRenderTarget2DImageViewer.h b/Engine/Plugins/Experimental/MediaViewer/Source/MediaViewer/Private/ImageViewers/TextureRenderTarget2DImageViewer.h
--- a/Engine/Plugins/Experimental/MediaViewer/Source/MediaViewer/Private/ImageViewers/TextureRenderTarget2DImageViewer.h
+++ b/Engine/Plugins/Experimental/MediaViewer/Source/MediaViewer/Private/ImageViewers/TextureRenderTarget2DImageViewer.h
@@ -20,7 +20,7 @@ struct FTextureRenderTarget2DImageViewerSettings
 	GENERATED_BODY()
 
 	UPROPERTY(VisibleAnywhere, Category = "Texture")
-	TObjectPtr<UTextureRenderTarget2D> RenderTarget;
+	TObjectPtr<UTextureRenderTarget2D> RenderTarget = nullptr;
 };
 
 namespace UE::MediaViewer::Private
diff --git a/Engine/Plugins/Experimental/MediaViewer/Source/MediaViewer/Private/ImageViewers/TextureSampleCache.cpp b/Engine/Plugins/Experimental/MediaViewer/Source/MediaViewer/Private/ImageViewers/TextureSampleCache.cpp
--- a/Engine/Plugins/Experimental/MediaViewer/Source/MediaViewer/Private/ImageViewers/TextureSampleCache.cpp
+++ b/Engine/Plugins/Experimental/MediaViewer/Source/MediaViewer/Private/ImageViewers/TextureSampleCache.cpp
@@ -4,6 +4,7 @@
 
 #include "Engine/Texture.h"
 #include "ImageViewers/TextureMipCache.h"
+#include "Misc/ScopeLock.h"
 #include "PixelFormat.h"
 #include "RenderingThread.h"
 #include "TextureResource.h"
@@ -35,6 +36,7 @@ FTextureSampleCache::FTextureSampleCache()
 FTextureSampleCache::FTextureSampleCache(TNotNull<UTexture*> InTexture, EPixelFormat InPixelFormat)
 	: Texture(InTexture)
 	, PixelFormat(InPixelFormat)
+	, SampleCS(FCriticalSection())
 {
 	if (!CanCopyTextureWithRHICommand(PixelFormat))
 	{
@@ -47,6 +49,13 @@ bool FTextureSampleCache::IsValid() const
 	return ::IsValid(Texture);
 }
 
+bool FTextureSampleCache::NeedsUpdate(const FIntPoint& InPixelCoordinates, const TOptional<FTimespan>& InTime) const
+{
+	FScopeLock Lock(&SampleCS);
+	return bDirty || !PixelColorSample.IsSet() || InPixelCoordinates != PixelColorSample->Coordinates
+		|| (InTime.IsSet() && (PixelColorSample->Time != InTime.GetValue()));
+}
+
 const FLinearColor* FTextureSampleCache::GetPixelColor(const FIntPoint& InPixelCoordinates, const TOptional<FTimespan>& InTime)
 {
 	if (!Texture)
@@ -54,8 +63,7 @@ const FLinearColor* FTextureSampleCache::GetPixelColor(const FIntPoint& InPixelC
 		return nullptr;
 	}
 
-	const bool bNeedsUpdate = bDirty || !PixelColorSample.IsSet() || InPixelCoordinates != PixelColorSample->Coordinates
-		|| (InTime.IsSet() && (PixelColorSample->Time != InTime.GetValue()));
+	const bool bNeedsUpdate = NeedsUpdate(InPixelCoordinates, InTime);
 
 	if (bNeedsUpdate)
 	{
@@ -69,9 +77,13 @@ const FLinearColor* FTextureSampleCache::GetPixelColor(const FIntPoint& InPixelC
 		}
 	}
 
-	if (PixelColorSample.IsSet())
 	{
-		return &PixelColorSample->Color;
+		FScopeLock Lock(&SampleCS);
+
+		if (PixelColorSample.IsSet())
+		{
+			return &PixelColorSample->Color;
+		}
 	}
 
 	return nullptr;
@@ -79,12 +91,16 @@ const FLinearColor* FTextureSampleCache::GetPixelColor(const FIntPoint& InPixelC
 
 void FTextureSampleCache::MarkDirty()
 {
+	FScopeLock Lock(&SampleCS);
 	bDirty = true;
 }
 
 void FTextureSampleCache::Invalidate()
 {
-	PixelColorSample.Reset();
+	{
+		FScopeLock Lock(&SampleCS);
+		PixelColorSample.Reset();
+	}
 
 	if (MipCache.IsValid())
 	{
@@ -92,32 +108,57 @@ void FTextureSampleCache::Invalidate()
 	}
 }
 
+FTextureSampleCache& FTextureSampleCache::operator=(const FTextureSampleCache& InOther)
+{
+	FScopeLock MyLock(&SampleCS);
+	FScopeLock OtherLock(&InOther.SampleCS);
+
+	Texture = InOther.Texture;
+	PixelFormat = InOther.PixelFormat;
+	bDirty = InOther.bDirty;
+	PixelColorSample = InOther.PixelColorSample;
+	MipCache = InOther.MipCache;
+	
+	return *this;
+}
+
 void FTextureSampleCache::SetPixelColor_RHI(const FIntPoint& InPixelCoordinates, const TOptional<FTimespan>& InTime)
 {
 	ENQUEUE_RENDER_COMMAND(GetPixelColors)
 	(
 		[this, InPixelCoordinates, InTime](FRHICommandListImmediate& RHICmdList)
 		{
-			if (FTextureResource* TextureResource = Texture->GetResource())
+			if (IsValid())
 			{
-				const FIntVector Size = TextureResource->TextureRHI->GetSizeXYZ();
-				TArray<FLinearColor> Data;
-				FIntRect Rect(InPixelCoordinates.X, InPixelCoordinates.Y, InPixelCoordinates.X + 1, InPixelCoordinates.Y + 1);
-				RHICmdList.ReadSurfaceData(TextureResource->TextureRHI, Rect, Data, FReadSurfaceDataFlags());
-
-				if (!Data.IsEmpty())
+				if (FTextureResource* TextureResource = Texture->GetResource())
 				{
-					PixelColorSample = {
-						InPixelCoordinates,
-						InTime,
-						Data[0]
-					};
-					bDirty = false;
-					return;
+					const FIntVector Size = TextureResource->TextureRHI->GetSizeXYZ();
+					TArray<FLinearColor> Data;
+					FIntRect Rect(InPixelCoordinates.X, InPixelCoordinates.Y, InPixelCoordinates.X + 1, InPixelCoordinates.Y + 1);
+					RHICmdList.ReadSurfaceData(TextureResource->TextureRHI, Rect, Data, FReadSurfaceDataFlags());
+
+					if (!Data.IsEmpty())
+					{
+						{
+							FScopeLock Lock(&SampleCS);
+
+							PixelColorSample = {
+								InPixelCoordinates,
+								InTime,
+								Data[0]
+							};
+
+							bDirty = false;
+						}
+						return;
+					}
 				}
 			}
 
-			PixelColorSample.Reset();
+			{
+				FScopeLock Lock(&SampleCS);
+				PixelColorSample.Reset();
+			}
 		}
 	);
 }
@@ -129,22 +170,26 @@ void FTextureSampleCache::SetPixelColor_Mip(const FIntPoint& InPixelCoordinates,
 		return;
 	}
 
-	PixelColorSample.Reset();
-
 	const FImage* Mip = MipCache->GetMipImage(/* Mip level */ 0);
 
 	if (!Mip)
 	{
+		FScopeLock Lock(&SampleCS);
+		PixelColorSample.Reset();
 		return;
 	}
 
-	PixelColorSample = {
-		InPixelCoordinates,
-		InTime,
-		Mip->GetOnePixelLinear(InPixelCoordinates.X, InPixelCoordinates.Y)
-	};
+	{
+		FScopeLock Lock(&SampleCS);
 
-	bDirty = false;
+		PixelColorSample = {
+			InPixelCoordinates,
+			InTime,
+			Mip->GetOnePixelLinear(InPixelCoordinates.X, InPixelCoordinates.Y)
+		};
+
+		bDirty = false;
+	}
 }
 
 } // UE::MediaViewer::Private
diff --git a/Engine/Plugins/Experimental/MediaViewer/Source/MediaViewer/Private/ImageViewers/TextureSampleCache.h b/Engine/Plugins/Experimental/MediaViewer/Source/MediaViewer/Private/ImageViewers/TextureSampleCache.h
--- a/Engine/Plugins/Experimental/MediaViewer/Source/MediaViewer/Private/ImageViewers/TextureSampleCache.h
+++ b/Engine/Plugins/Experimental/MediaViewer/Source/MediaViewer/Private/ImageViewers/TextureSampleCache.h
@@ -4,6 +4,7 @@
 
 #include "Templates/SharedPointer.h"
 
+#include "HAL/CriticalSection.h"
 #include "Math/Color.h"
 #include "Math/IntPoint.h"
 #include "Misc/NotNull.h"
@@ -32,6 +33,8 @@ public:
 
 	void Invalidate();
 
+	FTextureSampleCache& operator=(const FTextureSampleCache& InOther);
+
 protected:
 	struct FPixelColorSample
 	{
@@ -45,10 +48,13 @@ protected:
 	bool bDirty = true;
 	TOptional<FPixelColorSample> PixelColorSample;
 	TSharedPtr<FTextureMipCache> MipCache;
+	mutable FCriticalSection SampleCS;
 
 	void SetPixelColor_RHI(const FIntPoint& InPixelCoordinates, const TOptional<FTimespan>& InTime);
 
 	void SetPixelColor_Mip(const FIntPoint& InPixelCoordinates, const TOptional<FTimespan>& InTime);
+
+	bool NeedsUpdate(const FIntPoint& InPixelCoordinates, const TOptional<FTimespan>& InTime) const;
 };
 
 } // UE::MediaViewer::Private
diff --git a/Engine/Plugins/Experimental/MediaViewer/Source/MediaViewer/Private/ImageViewers/ViewportImageViewer.cpp b/Engine/Plugins/Experimental/MediaViewer/Source/MediaViewer/Private/ImageViewers/ViewportImageViewer.cpp
--- a/Engine/Plugins/Experimental/MediaViewer/Source/MediaViewer/Private/ImageViewers/ViewportImageViewer.cpp
+++ b/Engine/Plugins/Experimental/MediaViewer/Source/MediaViewer/Private/ImageViewers/ViewportImageViewer.cpp
@@ -19,12 +19,16 @@ namespace UE::MediaViewer::Private
 
 FViewportImageViewer::FViewportImageViewer(const FMediaImageViewerInfo& InImageInfo)
 	: FMediaImageViewer(InImageInfo)
+	, RenderState(ERenderState::NotRendered)
+	, bResizeTargetToViewport(false)
 {
 	DrawEffects |= ESlateDrawEffect::InvertAlpha | ESlateDrawEffect::PreMultipliedAlpha | ESlateDrawEffect::NoGamma;
 }
 
 FViewportImageViewer::~FViewportImageViewer()
 {
+	// Not removing bind to FCoreDelegates::OnEndFrameRT here because it can cause threading issues.
+	// It will be compacted and removed automatically by the delegate.
 }
 
 bool FViewportImageViewer::IsValid() const
@@ -55,6 +59,7 @@ TOptional<TVariant<FColor, FLinearColor>> FViewportImageViewer::GetPixelColor(co
 	{
 		TVariant<FColor, FLinearColor> PixelColorVariant;
 		PixelColorVariant.Set<FLinearColor>(*PixelColor);
+		PixelColorVariant.Get<FLinearColor>().A = 1.f - PixelColorVariant.Get<FLinearColor>().A;
 
 		return PixelColorVariant;
 	}
@@ -66,15 +71,15 @@ TSharedPtr<FStructOnScope> FViewportImageViewer::GetCustomSettingsOnScope() cons
 {
 	return MakeShared<FStructOnScope>(
 		FViewportImageViewerSettings::StaticStruct(), 
-		reinterpret_cast<uint8*>(&const_cast<FViewportImageViewerSettings&>(ViweportSettings))
+		reinterpret_cast<uint8*>(&const_cast<FViewportImageViewerSettings&>(ViewportSettings))
 	);
 }
 
 void FViewportImageViewer::PaintImage(FMediaImagePaintParams& InPaintParams, const FMediaImagePaintGeometry& InPaintGeometry)
 {
-	if (ViweportSettings.bRealTime)
+	if (RenderState == ERenderState::NotRendered || ViewportSettings.bRealTime)
 	{
-		RenderViewport();
+		RequestRender(/* Update render target size */ false);
 	}
 
 	FMediaImageViewer::PaintImage(InPaintParams, InPaintGeometry);
@@ -90,7 +95,7 @@ void FViewportImageViewer::AddReferencedObjects(FReferenceCollector& Collector)
 {
 	Collector.AddPropertyReferencesWithStructARO(
 		FViewportImageViewerSettings::StaticStruct(),
-		&ViweportSettings
+		&ViewportSettings
 	);
 }
 
@@ -106,15 +111,15 @@ void FViewportImageViewer::CreateBrush()
 	ImageInfo.Size = Viewport->GetSizeXY();
 
 	UTextureRenderTarget2D* RenderTarget = FMediaViewerUtils::CreateRenderTarget(ImageInfo.Size, /* Transparent */ false);
-	ViweportSettings.RenderTarget = RenderTarget;
+	ViewportSettings.RenderTarget = RenderTarget;
 
 	Brush = MakeShared<FSlateImageBrush>(RenderTarget, FVector2D(RenderTarget->GetSurfaceWidth(), RenderTarget->GetSurfaceDepth()));
 
 	SampleCache = FTextureSampleCache(RenderTarget, RenderTarget->GetFormat());
 }
 
-void FViewportImageViewer::RenderViewport(TNotNull<FViewport*> InViewport, TNotNull<UTextureRenderTarget2D*> InRenderTarget, 
-	bool bResizeTargetToViewport)
+void FViewportImageViewer::RenderViewport(TNotNull<FViewport*> InViewport, TNotNull<UTextureRenderTarget2D*> InRenderTarget,
+	FRenderComplete InRenderComplete, bool bResizeTargetToViewport)
 {
 	if (bResizeTargetToViewport)
 	{
@@ -128,12 +133,13 @@ void FViewportImageViewer::RenderViewport(TNotNull<FViewport*> InViewport, TNotN
 	}
 
 	ENQUEUE_RENDER_COMMAND(CopyViewportRenderTarget)(
-		[Viewport = InViewport, RenderTarget = InRenderTarget](FRHICommandListImmediate& InRHICmdList)
+		[Viewport = InViewport, RenderTarget = InRenderTarget, InRenderComplete](FRHICommandListImmediate& InRHICmdList)
 		{
 			FTextureRenderTargetResource* RenderTargetResource = RenderTarget->GetRenderTargetResource();
 
 			if (!RenderTargetResource)
 			{
+				InRenderComplete.ExecuteIfBound(false);
 				return;
 			}
 
@@ -143,6 +149,7 @@ void FViewportImageViewer::RenderViewport(TNotNull<FViewport*> InViewport, TNotN
 
 			if (!InputTexture)
 			{
+				InRenderComplete.ExecuteIfBound(false);
 				return;
 			}
 
@@ -150,6 +157,7 @@ void FViewportImageViewer::RenderViewport(TNotNull<FViewport*> InViewport, TNotN
 
 			if (!OutputTexture)
 			{
+				InRenderComplete.ExecuteIfBound(false);
 				return;
 			}
 
@@ -158,24 +166,56 @@ void FViewportImageViewer::RenderViewport(TNotNull<FViewport*> InViewport, TNotN
 			AddDrawTexturePass(GraphBuilder, GlobalShaderMap, InputTexture, OutputTexture, FRDGDrawTextureInfo());
 
 			GraphBuilder.Execute();
+
+			InRenderComplete.ExecuteIfBound(true);
 		}
 	);
 }
 
-void FViewportImageViewer::RenderViewport()
+void FViewportImageViewer::RequestRender(bool bInResizeTargetToViewport)
+{
+	RenderState = ERenderState::Invalidated;
+	bResizeTargetToViewport = bResizeTargetToViewport || bInResizeTargetToViewport;
+
+	if (!OnEndFrameRTHandle.IsValid())
+	{
+		OnEndFrameRTHandle = FCoreDelegates::OnEndFrameRT.AddSP(
+			this,
+			&FViewportImageViewer::ConditionallyRequestRender
+		);
+	}
+}
+
+void FViewportImageViewer::ConditionallyRequestRender()
 {
+	if (RenderState == ERenderState::RenderQueued)
+	{
+		return;
+	}
+
 	FViewport* Viewport = GetViewport();
 
-	if (!Viewport || !::IsValid(ViweportSettings.RenderTarget))
+	if (!Viewport || !::IsValid(ViewportSettings.RenderTarget))
 	{
 		return;
 	}
 
 	ImageInfo.Size = Viewport->GetSizeXY();
 
-	RenderViewport(TNotNull<FViewport*>(Viewport), TNotNull<UTextureRenderTarget2D*>(ViweportSettings.RenderTarget));
+	RenderViewport(
+		TNotNull<FViewport*>(Viewport), 
+		TNotNull<UTextureRenderTarget2D*>(ViewportSettings.RenderTarget),
+		FRenderComplete::CreateSP(this, &FViewportImageViewer::OnRenderComplete),
+		bResizeTargetToViewport
+	);
 
-	if (SampleCache.IsValid())
+	RenderState = ERenderState::RenderQueued;
+	bResizeTargetToViewport = false;
+}
+
+void FViewportImageViewer::OnRenderComplete(bool bInSuccess)
+{
+	if (bInSuccess)
 	{
 		SampleCache.MarkDirty();
 	}
diff --git a/Engine/Plugins/Experimental/MediaViewer/Source/MediaViewer/Private/ImageViewers/ViewportImageViewer.h b/Engine/Plugins/Experimental/MediaViewer/Source/MediaViewer/Private/ImageViewers/ViewportImageViewer.h
--- a/Engine/Plugins/Experimental/MediaViewer/Source/MediaViewer/Private/ImageViewers/ViewportImageViewer.h
+++ b/Engine/Plugins/Experimental/MediaViewer/Source/MediaViewer/Private/ImageViewers/ViewportImageViewer.h
@@ -6,6 +6,7 @@
 #include "ImageViewer/MediaImageViewer.h"
 #include "UObject/GCObject.h"
 
+#include "Delegates/IDelegateInstance.h"
 #include "ImageViewers/TextureSampleCache.h"
 #include "Templates/SharedPointer.h"
 
@@ -21,11 +22,11 @@ struct FViewportImageViewerSettings
 	GENERATED_BODY()
 
 	UPROPERTY(VisibleAnywhere, Category = "Material")
-	TObjectPtr<UTextureRenderTarget2D> RenderTarget;
+	TObjectPtr<UTextureRenderTarget2D> RenderTarget = nullptr;
 
 	/** Will render the material every frame. */
 	UPROPERTY(EditAnywhere, Category = "Material")
-	bool bRealTime;
+	bool bRealTime = true;
 };
 
 namespace UE::MediaViewer::Private
@@ -50,16 +51,32 @@ public:
 	//~ End FGCObject
 
 protected:
-	static void RenderViewport(TNotNull<FViewport*> InViewport, TNotNull<UTextureRenderTarget2D*> InRenderTarget, bool bResizeTargetToViewport = true);
-
-	FViewportImageViewerSettings ViweportSettings;
+	enum class ERenderState : uint8
+	{
+		NotRendered,
+		RenderQueued,
+		Invalidated
+	};
+
+	DECLARE_DELEGATE_OneParam(FRenderComplete, bool /* Was succsesfully rendered */)
+	static void RenderViewport(TNotNull<FViewport*> InViewport, TNotNull<UTextureRenderTarget2D*> InRenderTarget, 
+		FRenderComplete InRenderComplete, bool bResizeTargetToViewport = true);
+
+	FViewportImageViewerSettings ViewportSettings;
 	mutable FTextureSampleCache SampleCache;
+	FDelegateHandle OnEndFrameRTHandle;
+	ERenderState RenderState;
+	bool bResizeTargetToViewport;
 
 	virtual FViewport* GetViewport() const = 0;
 
 	void CreateBrush();
 
-	void RenderViewport();
+	void RequestRender(bool bInResizeTargetToViewport);
+	
+	void ConditionallyRequestRender();
+
+	void OnRenderComplete(bool bInSuccess);
 
 	//~ Begin FMediaImageViewer
 	virtual void PaintImage(FMediaImagePaintParams& InPaintParams, const FMediaImagePaintGeometry& InPaintGeometry) override;
