Media Viewer: Cross-thread hardening for pixel sampling

diff --git a/Engine/Plugins/Experimental/MediaViewer/Source/MediaViewer/Private/ImageViewers/MaterialInterfaceImageViewer.cpp b/Engine/Plugins/Experimental/MediaViewer/Source/MediaViewer/Private/ImageViewers/MaterialInterfaceImageViewer.cpp
--- a/Engine/Plugins/Experimental/MediaViewer/Source/MediaViewer/Private/ImageViewers/MaterialInterfaceImageViewer.cpp
+++ b/Engine/Plugins/Experimental/MediaViewer/Source/MediaViewer/Private/ImageViewers/MaterialInterfaceImageViewer.cpp
@@ -282,24 +282,24 @@ TSharedPtr<FMediaViewerLibraryItem> FMaterialInterfaceImageViewer::CreateLibrary
 
 TOptional<TVariant<FColor, FLinearColor>> FMaterialInterfaceImageViewer::GetPixelColor(const FIntPoint& InPixelCoords, int32 InMipLevel) const
 {
-	if (!SampleCache.IsValid())
+	if (!SampleCache.IsValid() || !SampleCache->IsValid())
 	{
 		return {};
 	}
 
 	if (InPixelCoords.X < 0 || InPixelCoords.Y < 0)
 	{
-		SampleCache.Invalidate();
+		SampleCache->Invalidate();
 		return {};
 	}
 
 	if (InPixelCoords.X >= ImageInfo.Size.X || InPixelCoords.Y >= ImageInfo.Size.Y)
 	{
-		SampleCache.Invalidate();
+		SampleCache->Invalidate();
 		return {};
 	}
 
-	if (const FLinearColor* PixelColor = SampleCache.GetPixelColor(InPixelCoords))
+	if (const FLinearColor* PixelColor = SampleCache->GetPixelColor(InPixelCoords))
 	{
 		TVariant<FColor, FLinearColor> PixelColorVariant;
 		PixelColorVariant.Set<FLinearColor>(*PixelColor);
@@ -362,7 +362,7 @@ void FMaterialInterfaceImageViewer::CreateBrush()
 
 	Brush = MakeShared<FSlateImageBrush>(RenderTarget, FVector2D(ImageInfo.Size.X, ImageInfo.Size.Y));
 
-	SampleCache = FTextureSampleCache(RenderTarget, RenderTarget->GetFormat());
+	SampleCache = MakeShared<FTextureSampleCache>(RenderTarget, RenderTarget->GetFormat());
 }
 
 void FMaterialInterfaceImageViewer::RenderMaterial()
@@ -373,7 +373,7 @@ void FMaterialInterfaceImageViewer::RenderMaterial()
 
 		if (SampleCache.IsValid())
 		{
-			SampleCache.MarkDirty();
+			SampleCache->MarkDirty();
 		}
 	}
 }
diff --git a/Engine/Plugins/Experimental/MediaViewer/Source/MediaViewer/Private/ImageViewers/MaterialInterfaceImageViewer.h b/Engine/Plugins/Experimental/MediaViewer/Source/MediaViewer/Private/ImageViewers/MaterialInterfaceImageViewer.h
--- a/Engine/Plugins/Experimental/MediaViewer/Source/MediaViewer/Private/ImageViewers/MaterialInterfaceImageViewer.h
+++ b/Engine/Plugins/Experimental/MediaViewer/Source/MediaViewer/Private/ImageViewers/MaterialInterfaceImageViewer.h
@@ -113,7 +113,7 @@ public:
 
 protected:
 	FMaterialInterfaceImageViewerSettings MaterialSettings;
-	mutable FTextureSampleCache SampleCache;
+	TSharedPtr<FTextureSampleCache> SampleCache;
 
 	void CreateBrush();
 
diff --git a/Engine/Plugins/Experimental/MediaViewer/Source/MediaViewer/Private/ImageViewers/MediaSourceImageViewer.cpp b/Engine/Plugins/Experimental/MediaViewer/Source/MediaViewer/Private/ImageViewers/MediaSourceImageViewer.cpp
--- a/Engine/Plugins/Experimental/MediaViewer/Source/MediaViewer/Private/ImageViewers/MediaSourceImageViewer.cpp
+++ b/Engine/Plugins/Experimental/MediaViewer/Source/MediaViewer/Private/ImageViewers/MediaSourceImageViewer.cpp
@@ -117,7 +117,7 @@ FMediaSourceImageViewer::FItem::FItem(const FText& InName, const FText& InToolTi
 }
 
 FMediaSourceImageViewer::FItem::FItem(const FGuid& InId, const FText& InName, const FText& InToolTip, bool bInTransient, const FString& InStringValue)
-	: FMediaViewerLibraryItem(InName, InToolTip, bInTransient, InStringValue)
+	: FMediaViewerLibraryItem(InId, InName, InToolTip, bInTransient, InStringValue)
 {
 
 }
@@ -294,7 +294,7 @@ FMediaSourceImageViewer::FMediaSourceImageViewer(const FGuid& InId, TNotNull<UMe
 				MediaTexture->AddAssetUserData(NewObject<UMediaViewerUserData>(MediaTexture));
 
 				// We can't know the pixel format of the media texture at this point.
-				SampleCache = FTextureSampleCache(MediaTexture, PF_Unknown);
+				SampleCache = MakeShared<FTextureSampleCache>(MediaTexture, PF_Unknown);
 			}
 		}
 	}
@@ -346,26 +346,26 @@ TSharedPtr<FMediaViewerLibraryItem> FMediaSourceImageViewer::CreateLibraryItem()
 
 TOptional<TVariant<FColor, FLinearColor>> FMediaSourceImageViewer::GetPixelColor(const FIntPoint& InPixelCoords, int32 InMipLevel) const
 {
-	if (!SampleCache.IsValid())
+	if (!SampleCache.IsValid() || !SampleCache->IsValid())
 	{
 		return {};
 	}
 	
 	if (InPixelCoords.X < 0 || InPixelCoords.Y < 0)
 	{
-		SampleCache.Invalidate();
+		SampleCache->Invalidate();
 		return {};
 	}
 
 	if (InPixelCoords.X >= ImageInfo.Size.X || InPixelCoords.Y >= ImageInfo.Size.Y)
 	{
-		SampleCache.Invalidate();
+		SampleCache->Invalidate();
 		return {};
 	}
 
 	if (!MediaStream)
 	{
-		SampleCache.Invalidate();
+		SampleCache->Invalidate();
 		return {};
 	}
 
@@ -373,7 +373,7 @@ TOptional<TVariant<FColor, FLinearColor>> FMediaSourceImageViewer::GetPixelColor
 
 	if (!MediaStreamPlayer)
 	{
-		SampleCache.Invalidate();
+		SampleCache->Invalidate();
 		return {};
 	}
 
@@ -381,7 +381,7 @@ TOptional<TVariant<FColor, FLinearColor>> FMediaSourceImageViewer::GetPixelColor
 	{
 		const FTimespan PlayerTime = Player->GetTime();
 
-		if (const FLinearColor* PixelColor = SampleCache.GetPixelColor(InPixelCoords, PlayerTime))
+		if (const FLinearColor* PixelColor = SampleCache->GetPixelColor(InPixelCoords, PlayerTime))
 		{
 			TVariant<FColor, FLinearColor> PixelColorVariant;
 			PixelColorVariant.Set<FLinearColor>(*PixelColor);
diff --git a/Engine/Plugins/Experimental/MediaViewer/Source/MediaViewer/Private/ImageViewers/MediaSourceImageViewer.h b/Engine/Plugins/Experimental/MediaViewer/Source/MediaViewer/Private/ImageViewers/MediaSourceImageViewer.h
--- a/Engine/Plugins/Experimental/MediaViewer/Source/MediaViewer/Private/ImageViewers/MediaSourceImageViewer.h
+++ b/Engine/Plugins/Experimental/MediaViewer/Source/MediaViewer/Private/ImageViewers/MediaSourceImageViewer.h
@@ -133,7 +133,7 @@ public:
 protected:
 	FMediaSourceImageViewerSettings MediaSourceSettings;
 	TObjectPtr<UMediaStream> MediaStream;
-	mutable FTextureSampleCache SampleCache;
+	TSharedPtr<FTextureSampleCache> SampleCache;
 
 	void AddPlayerName(const TSharedRef<SHorizontalBox>& InStatusBar);
 
diff --git a/Engine/Plugins/Experimental/MediaViewer/Source/MediaViewer/Private/ImageViewers/MediaTextureImageViewer.cpp b/Engine/Plugins/Experimental/MediaViewer/Source/MediaViewer/Private/ImageViewers/MediaTextureImageViewer.cpp
--- a/Engine/Plugins/Experimental/MediaViewer/Source/MediaViewer/Private/ImageViewers/MediaTextureImageViewer.cpp
+++ b/Engine/Plugins/Experimental/MediaViewer/Source/MediaViewer/Private/ImageViewers/MediaTextureImageViewer.cpp
@@ -183,7 +183,7 @@ FMediaTextureImageViewer::FMediaTextureImageViewer(const FGuid& InId, TNotNull<U
 	);
 
 	// We can't know the pixel format of the media texture at this point.
-	SampleCache = FTextureSampleCache(InMediaTexture, PF_Unknown);
+	SampleCache = MakeShared<FTextureSampleCache>(InMediaTexture, PF_Unknown);
 }
 
 TSharedPtr<FMediaViewerLibraryItem> FMediaTextureImageViewer::CreateLibraryItem() const
@@ -206,27 +206,27 @@ TSharedPtr<FMediaViewerLibraryItem> FMediaTextureImageViewer::CreateLibraryItem(
 
 TOptional<TVariant<FColor, FLinearColor>> FMediaTextureImageViewer::GetPixelColor(const FIntPoint& InPixelCoords, int32 InMipLevel) const
 {
-	if (!SampleCache.IsValid())
+	if (!SampleCache.IsValid() || !SampleCache->IsValid())
 	{
 		return {};
 	}
 
 	// Always mark the sample cache dirty. We have no control over the state of the media texture.
-	SampleCache.MarkDirty();
+	SampleCache->MarkDirty();
 
 	if (InPixelCoords.X < 0 || InPixelCoords.Y < 0)
 	{
-		SampleCache.Invalidate();
+		SampleCache->Invalidate();
 		return {};
 	}
 
 	if (InPixelCoords.X >= ImageInfo.Size.X || InPixelCoords.Y >= ImageInfo.Size.Y)
 	{
-		SampleCache.Invalidate();
+		SampleCache->Invalidate();
 		return {};
 	}
 
-	if (const FLinearColor* PixelColor = SampleCache.GetPixelColor(InPixelCoords))
+	if (const FLinearColor* PixelColor = SampleCache->GetPixelColor(InPixelCoords))
 	{
 		TVariant<FColor, FLinearColor> PixelColorVariant;
 		PixelColorVariant.Set<FLinearColor>(*PixelColor);
diff --git a/Engine/Plugins/Experimental/MediaViewer/Source/MediaViewer/Private/ImageViewers/MediaTextureImageViewer.h b/Engine/Plugins/Experimental/MediaViewer/Source/MediaViewer/Private/ImageViewers/MediaTextureImageViewer.h
--- a/Engine/Plugins/Experimental/MediaViewer/Source/MediaViewer/Private/ImageViewers/MediaTextureImageViewer.h
+++ b/Engine/Plugins/Experimental/MediaViewer/Source/MediaViewer/Private/ImageViewers/MediaTextureImageViewer.h
@@ -90,7 +90,7 @@ public:
 
 protected:
 	FMediaTextureImageViewerSettings MediaTextureSettings;
-	mutable FTextureSampleCache SampleCache;
+	TSharedPtr<FTextureSampleCache> SampleCache;
 };
 
 } // UE::MediaViewer::Private
diff --git a/Engine/Plugins/Experimental/MediaViewer/Source/MediaViewer/Private/ImageViewers/Texture2DImageViewer.cpp b/Engine/Plugins/Experimental/MediaViewer/Source/MediaViewer/Private/ImageViewers/Texture2DImageViewer.cpp
--- a/Engine/Plugins/Experimental/MediaViewer/Source/MediaViewer/Private/ImageViewers/Texture2DImageViewer.cpp
+++ b/Engine/Plugins/Experimental/MediaViewer/Source/MediaViewer/Private/ImageViewers/Texture2DImageViewer.cpp
@@ -199,7 +199,7 @@ FTexture2DImageViewer::FTexture2DImageViewer(const FGuid& InId, TNotNull<UTextur
 		InTexture->GetPlatformData()->Mips.Num(),
 		FMediaImageViewer::GetObjectDisplayName(InTexture)
 	})
-	, SampleCache(InTexture, InTexture->GetPixelFormat())
+	, SampleCache(MakeShared<FTextureSampleCache>(InTexture, InTexture->GetPixelFormat()))
 	, bValidImageSize(false)
 {
 	TextureSettings.Texture = InTexture;
@@ -235,7 +235,7 @@ TSharedPtr<FMediaViewerLibraryItem> FTexture2DImageViewer::CreateLibraryItem() c
 
 TOptional<TVariant<FColor, FLinearColor>> FTexture2DImageViewer::GetPixelColor(const FIntPoint& InPixelCoords, int32 InMipLevel) const
 {
-	if (!SampleCache.IsValid())
+	if (!SampleCache.IsValid() || !SampleCache->IsValid())
 	{
 		return {};
 	}
@@ -250,7 +250,7 @@ TOptional<TVariant<FColor, FLinearColor>> FTexture2DImageViewer::GetPixelColor(c
 		return {};
 	}
 
-	if (const FLinearColor* PixelColor = SampleCache.GetPixelColor(InPixelCoords))
+	if (const FLinearColor* PixelColor = SampleCache->GetPixelColor(InPixelCoords))
 	{
 		TVariant<FColor, FLinearColor> PixelColorVariant;
 		PixelColorVariant.Set<FLinearColor>(*PixelColor);
diff --git a/Engine/Plugins/Experimental/MediaViewer/Source/MediaViewer/Private/ImageViewers/Texture2DImageViewer.h b/Engine/Plugins/Experimental/MediaViewer/Source/MediaViewer/Private/ImageViewers/Texture2DImageViewer.h
--- a/Engine/Plugins/Experimental/MediaViewer/Source/MediaViewer/Private/ImageViewers/Texture2DImageViewer.h
+++ b/Engine/Plugins/Experimental/MediaViewer/Source/MediaViewer/Private/ImageViewers/Texture2DImageViewer.h
@@ -96,7 +96,7 @@ public:
 
 protected:
 	FTexture2DImageViewerSettings TextureSettings;
-	mutable FTextureSampleCache SampleCache;
+	TSharedPtr<FTextureSampleCache> SampleCache;
 	bool bValidImageSize;
 
 	TOptional<FIntPoint> GetTextureSize() const;
diff --git a/Engine/Plugins/Experimental/MediaViewer/Source/MediaViewer/Private/ImageViewers/TextureRenderTarget2DImageViewer.cpp b/Engine/Plugins/Experimental/MediaViewer/Source/MediaViewer/Private/ImageViewers/TextureRenderTarget2DImageViewer.cpp
--- a/Engine/Plugins/Experimental/MediaViewer/Source/MediaViewer/Private/ImageViewers/TextureRenderTarget2DImageViewer.cpp
+++ b/Engine/Plugins/Experimental/MediaViewer/Source/MediaViewer/Private/ImageViewers/TextureRenderTarget2DImageViewer.cpp
@@ -176,7 +176,7 @@ FTextureRenderTarget2DImageViewer::FTextureRenderTarget2DImageViewer(const FGuid
 		FVector2D(InRenderTarget->GetSurfaceWidth(), InRenderTarget->GetSurfaceHeight())
 	);
 	
-	SampleCache = FTextureSampleCache(InRenderTarget, InRenderTarget->GetFormat());
+	SampleCache = MakeShared<FTextureSampleCache>(InRenderTarget, InRenderTarget->GetFormat());
 }
 
 TSharedPtr<FMediaViewerLibraryItem> FTextureRenderTarget2DImageViewer::CreateLibraryItem() const
@@ -199,27 +199,27 @@ TSharedPtr<FMediaViewerLibraryItem> FTextureRenderTarget2DImageViewer::CreateLib
 
 TOptional<TVariant<FColor, FLinearColor>> FTextureRenderTarget2DImageViewer::GetPixelColor(const FIntPoint& InPixelCoords, int32 InMipLevel) const
 {
-	if (!SampleCache.IsValid())
+	if (!SampleCache.IsValid() || !SampleCache->IsValid())
 	{
 		return {};
 	}
 
 	// Always mark the sample cache dirty. We have no control over the state of the media texture.
-	SampleCache.MarkDirty();
+	SampleCache->MarkDirty();
 
 	if (InPixelCoords.X < 0 || InPixelCoords.Y < 0)
 	{
-		SampleCache.Invalidate();
+		SampleCache->Invalidate();
 		return {};
 	}
 
 	if (InPixelCoords.X >= ImageInfo.Size.X || InPixelCoords.Y >= ImageInfo.Size.Y)
 	{
-		SampleCache.Invalidate();
+		SampleCache->Invalidate();
 		return {};
 	}
 
-	if (const FLinearColor* PixelColor = SampleCache.GetPixelColor(InPixelCoords))
+	if (const FLinearColor* PixelColor = SampleCache->GetPixelColor(InPixelCoords))
 	{
 		TVariant<FColor, FLinearColor> PixelColorVariant;
 		PixelColorVariant.Set<FLinearColor>(*PixelColor);
diff --git a/Engine/Plugins/Experimental/MediaViewer/Source/MediaViewer/Private/ImageViewers/TextureRenderTarget2DImageViewer.h b/Engine/Plugins/Experimental/MediaViewer/Source/MediaViewer/Private/ImageViewers/TextureRenderTarget2DImageViewer.h
--- a/Engine/Plugins/Experimental/MediaViewer/Source/MediaViewer/Private/ImageViewers/TextureRenderTarget2DImageViewer.h
+++ b/Engine/Plugins/Experimental/MediaViewer/Source/MediaViewer/Private/ImageViewers/TextureRenderTarget2DImageViewer.h
@@ -87,7 +87,7 @@ public:
 
 protected:
 	FTextureRenderTarget2DImageViewerSettings RenderTargetSettings;
-	mutable FTextureSampleCache SampleCache;
+	TSharedPtr<FTextureSampleCache> SampleCache;
 };
 
 } // UE::MediaViewer::Private
diff --git a/Engine/Plugins/Experimental/MediaViewer/Source/MediaViewer/Private/ImageViewers/TextureSampleCache.cpp b/Engine/Plugins/Experimental/MediaViewer/Source/MediaViewer/Private/ImageViewers/TextureSampleCache.cpp
--- a/Engine/Plugins/Experimental/MediaViewer/Source/MediaViewer/Private/ImageViewers/TextureSampleCache.cpp
+++ b/Engine/Plugins/Experimental/MediaViewer/Source/MediaViewer/Private/ImageViewers/TextureSampleCache.cpp
@@ -126,38 +126,42 @@ void FTextureSampleCache::SetPixelColor_RHI(const FIntPoint& InPixelCoordinates,
 {
 	ENQUEUE_RENDER_COMMAND(GetPixelColors)
 	(
-		[this, InPixelCoordinates, InTime](FRHICommandListImmediate& RHICmdList)
+		[ThisWeak = SharedThis(this).ToWeakPtr(), InPixelCoordinates, InTime](FRHICommandListImmediate& RHICmdList)
 		{
-			if (IsValid())
+			TSharedPtr<FTextureSampleCache> This = ThisWeak.Pin();
+
+			if (!This.IsValid() || !This->IsValid())
 			{
-				if (FTextureResource* TextureResource = Texture->GetResource())
-				{
-					const FIntVector Size = TextureResource->TextureRHI->GetSizeXYZ();
-					TArray<FLinearColor> Data;
-					FIntRect Rect(InPixelCoordinates.X, InPixelCoordinates.Y, InPixelCoordinates.X + 1, InPixelCoordinates.Y + 1);
-					RHICmdList.ReadSurfaceData(TextureResource->TextureRHI, Rect, Data, FReadSurfaceDataFlags());
+				return;
+			}
+
+			if (FTextureResource* TextureResource = This->Texture->GetResource())
+			{
+				const FIntVector Size = TextureResource->TextureRHI->GetSizeXYZ();
+				TArray<FLinearColor> Data;
+				FIntRect Rect(InPixelCoordinates.X, InPixelCoordinates.Y, InPixelCoordinates.X + 1, InPixelCoordinates.Y + 1);
+				RHICmdList.ReadSurfaceData(TextureResource->TextureRHI, Rect, Data, FReadSurfaceDataFlags());
 
-					if (!Data.IsEmpty())
+				if (!Data.IsEmpty())
+				{
 					{
-						{
-							FScopeLock Lock(&SampleCS);
-
-							PixelColorSample = {
-								InPixelCoordinates,
-								InTime,
-								Data[0]
-							};
-
-							bDirty = false;
-						}
-						return;
+						FScopeLock Lock(&This->SampleCS);
+
+						This->PixelColorSample = {
+							InPixelCoordinates,
+							InTime,
+							Data[0]
+						};
+
+						This->bDirty = false;
 					}
+					return;
 				}
 			}
 
 			{
-				FScopeLock Lock(&SampleCS);
-				PixelColorSample.Reset();
+				FScopeLock Lock(&This->SampleCS);
+				This->PixelColorSample.Reset();
 			}
 		}
 	);
diff --git a/Engine/Plugins/Experimental/MediaViewer/Source/MediaViewer/Private/ImageViewers/TextureSampleCache.h b/Engine/Plugins/Experimental/MediaViewer/Source/MediaViewer/Private/ImageViewers/TextureSampleCache.h
--- a/Engine/Plugins/Experimental/MediaViewer/Source/MediaViewer/Private/ImageViewers/TextureSampleCache.h
+++ b/Engine/Plugins/Experimental/MediaViewer/Source/MediaViewer/Private/ImageViewers/TextureSampleCache.h
@@ -19,7 +19,7 @@ namespace UE::MediaViewer::Private
 
 class FTextureMipCache;
 
-class FTextureSampleCache
+class FTextureSampleCache : public TSharedFromThis<FTextureSampleCache>
 {
 public:
 	FTextureSampleCache();
diff --git a/Engine/Plugins/Experimental/MediaViewer/Source/MediaViewer/Private/ImageViewers/ViewportImageViewer.cpp b/Engine/Plugins/Experimental/MediaViewer/Source/MediaViewer/Private/ImageViewers/ViewportImageViewer.cpp
--- a/Engine/Plugins/Experimental/MediaViewer/Source/MediaViewer/Private/ImageViewers/ViewportImageViewer.cpp
+++ b/Engine/Plugins/Experimental/MediaViewer/Source/MediaViewer/Private/ImageViewers/ViewportImageViewer.cpp
@@ -38,24 +38,24 @@ bool FViewportImageViewer::IsValid() const
 
 TOptional<TVariant<FColor, FLinearColor>> FViewportImageViewer::GetPixelColor(const FIntPoint& InPixelCoords, int32 InMipLevel) const
 {
-	if (!SampleCache.IsValid())
+	if (!SampleCache.IsValid() || !SampleCache->IsValid())
 	{
 		return {};
 	}
 
 	if (InPixelCoords.X < 0 || InPixelCoords.Y < 0)
 	{
-		SampleCache.Invalidate();
+		SampleCache->Invalidate();
 		return {};
 	}
 
 	if (InPixelCoords.X >= ImageInfo.Size.X || InPixelCoords.Y >= ImageInfo.Size.Y)
 	{
-		SampleCache.Invalidate();
+		SampleCache->Invalidate();
 		return {};
 	}
 
-	if (const FLinearColor* PixelColor = SampleCache.GetPixelColor(InPixelCoords))
+	if (const FLinearColor* PixelColor = SampleCache->GetPixelColor(InPixelCoords))
 	{
 		TVariant<FColor, FLinearColor> PixelColorVariant;
 		PixelColorVariant.Set<FLinearColor>(*PixelColor);
@@ -115,7 +115,7 @@ void FViewportImageViewer::CreateBrush()
 
 	Brush = MakeShared<FSlateImageBrush>(RenderTarget, FVector2D(RenderTarget->GetSurfaceWidth(), RenderTarget->GetSurfaceDepth()));
 
-	SampleCache = FTextureSampleCache(RenderTarget, RenderTarget->GetFormat());
+	SampleCache = MakeShared<FTextureSampleCache>(RenderTarget, RenderTarget->GetFormat());
 }
 
 void FViewportImageViewer::RenderViewport(TNotNull<FViewport*> InViewport, TNotNull<UTextureRenderTarget2D*> InRenderTarget,
@@ -215,9 +215,9 @@ void FViewportImageViewer::ConditionallyRequestRender()
 
 void FViewportImageViewer::OnRenderComplete(bool bInSuccess)
 {
-	if (bInSuccess)
+	if (bInSuccess && SampleCache.IsValid())
 	{
-		SampleCache.MarkDirty();
+		SampleCache->MarkDirty();
 	}
 }
 
diff --git a/Engine/Plugins/Experimental/MediaViewer/Source/MediaViewer/Private/ImageViewers/ViewportImageViewer.h b/Engine/Plugins/Experimental/MediaViewer/Source/MediaViewer/Private/ImageViewers/ViewportImageViewer.h
--- a/Engine/Plugins/Experimental/MediaViewer/Source/MediaViewer/Private/ImageViewers/ViewportImageViewer.h
+++ b/Engine/Plugins/Experimental/MediaViewer/Source/MediaViewer/Private/ImageViewers/ViewportImageViewer.h
@@ -63,7 +63,7 @@ protected:
 		FRenderComplete InRenderComplete, bool bResizeTargetToViewport = true);
 
 	FViewportImageViewerSettings ViewportSettings;
-	mutable FTextureSampleCache SampleCache;
+	TSharedPtr<FTextureSampleCache> SampleCache;
 	FDelegateHandle OnEndFrameRTHandle;
 	ERenderState RenderState;
 	bool bResizeTargetToViewport;
