Media VIewer: Fixed changing splitter location in single image mode

diff --git a/Engine/Plugins/Experimental/MediaViewer/Source/MediaViewer/Private/Widgets/SMediaViewer.cpp b/Engine/Plugins/Experimental/MediaViewer/Source/MediaViewer/Private/Widgets/SMediaViewer.cpp
--- a/Engine/Plugins/Experimental/MediaViewer/Source/MediaViewer/Private/Widgets/SMediaViewer.cpp
+++ b/Engine/Plugins/Experimental/MediaViewer/Source/MediaViewer/Private/Widgets/SMediaViewer.cpp
@@ -681,8 +681,16 @@ float SMediaViewer::GetDPIScale() const
 
 void SMediaViewer::OnABResized(float InSize)
 {
-	StaticCastSharedRef<SSplitter>(ContentSlot->GetWidget())->SlotAt(0).SetSizeValue(InSize);
-	StaticCastSharedRef<SSplitter>(ContentSlot->GetWidget())->SlotAt(1).SetSizeValue(1.f - InSize);
+	if (!IsImageViewerNull(EMediaImageViewerPosition::First))
+	{
+		StaticCastSharedRef<SSplitter>(ContentSlot->GetWidget())->SlotAt(0).SetSizeValue(InSize);
+	}
+
+	if (!IsImageViewerNull(EMediaImageViewerPosition::Second))
+	{
+		StaticCastSharedRef<SSplitter>(ContentSlot->GetWidget())->SlotAt(1).SetSizeValue(1.f - InSize);
+	}
+
 	MediaViewerSettings.ABSplitterLocation = InSize * 100.f;
 }
 
