Media Viewer: Added more texture encoding types to the no-read-with-rhi list

diff --git a/Engine/Plugins/Experimental/MediaViewer/Source/MediaViewer/Private/ImageViewers/TextureSampleCache.cpp b/Engine/Plugins/Experimental/MediaViewer/Source/MediaViewer/Private/ImageViewers/TextureSampleCache.cpp
--- a/Engine/Plugins/Experimental/MediaViewer/Source/MediaViewer/Private/ImageViewers/TextureSampleCache.cpp
+++ b/Engine/Plugins/Experimental/MediaViewer/Source/MediaViewer/Private/ImageViewers/TextureSampleCache.cpp
@@ -17,10 +17,59 @@ bool CanCopyTextureWithRHICommand(EPixelFormat InFormat)
 {
 	switch (InFormat)
 	{
-		// DirectX compressed textures cannot read using the RHI ReadSurface method.
+		// Textures encoding that cannot be read by the RHI.
+		// Based on the list of invalid texture formats in the Metal RHI (see MetalRHI.cpp)
 		case EPixelFormat::PF_DXT1:
 		case EPixelFormat::PF_DXT3:
 		case EPixelFormat::PF_DXT5:
+		case EPixelFormat::PF_BC4:
+		case EPixelFormat::PF_BC5:
+		case EPixelFormat::PF_BC6H:
+		case EPixelFormat::PF_BC7:
+		case EPixelFormat::PF_DepthStencil:
+		case EPixelFormat::PF_ShadowDepth:
+		case EPixelFormat::PF_D24:
+		case EPixelFormat::PF_A1:
+		case EPixelFormat::PF_PVRTC2:
+		case EPixelFormat::PF_PVRTC4:
+		case EPixelFormat::PF_R5G6B5_UNORM:
+		case EPixelFormat::PF_B5G5R5A1_UNORM:
+		case EPixelFormat::PF_ATC_RGB:
+		case EPixelFormat::PF_ATC_RGBA_E:
+		case EPixelFormat::PF_ATC_RGBA_I:
+		case EPixelFormat::PF_X24_G8:
+		case EPixelFormat::PF_ETC1:
+		case EPixelFormat::PF_ETC2_RGB:
+		case EPixelFormat::PF_ETC2_RGBA:
+		case EPixelFormat::PF_ASTC_4x4:
+		case EPixelFormat::PF_ASTC_6x6:
+		case EPixelFormat::PF_ASTC_8x8:
+		case EPixelFormat::PF_ASTC_10x10:
+		case EPixelFormat::PF_ASTC_12x12:
+		case EPixelFormat::PF_ASTC_4x4_HDR:
+		case EPixelFormat::PF_ASTC_6x6_HDR:
+		case EPixelFormat::PF_ASTC_8x8_HDR:
+		case EPixelFormat::PF_ASTC_10x10_HDR:
+		case EPixelFormat::PF_ASTC_12x12_HDR:
+		case EPixelFormat::PF_L8:
+		case EPixelFormat::PF_R16G16B16A16_SNORM:
+		case EPixelFormat::PF_PLATFORM_HDR_0:
+		case EPixelFormat::PF_PLATFORM_HDR_1:
+		case EPixelFormat::PF_NV12:
+		case EPixelFormat::PF_ETC2_R11_EAC:
+		case EPixelFormat::PF_ETC2_RG11_EAC:
+		case EPixelFormat::PF_R32G32B32_UINT:
+		case EPixelFormat::PF_R32G32B32_SINT:
+		case EPixelFormat::PF_R32G32B32F:
+		case EPixelFormat::PF_R64_UINT:
+		case EPixelFormat::PF_R9G9B9EXP5:
+		case EPixelFormat::PF_P010:
+		case EPixelFormat::PF_ASTC_4x4_NORM_RG:
+		case EPixelFormat::PF_ASTC_6x6_NORM_RG:
+		case EPixelFormat::PF_ASTC_8x8_NORM_RG:
+		case EPixelFormat::PF_ASTC_10x10_NORM_RG:
+		case EPixelFormat::PF_ASTC_12x12_NORM_RG:
+		case EPixelFormat::PF_R8G8B8:
 			return false;
 
 		default:
