Media Viewer: Dragging and dropping update
- Library items that are assets are now asset drag+drops.
- Dragging items that get invalidate mid-drag will no longer crash the editor.

diff --git a/Engine/Plugins/Experimental/MediaViewer/Source/MediaViewer/Private/ImageViewers/LevelEditorViewportImageViewer.cpp b/Engine/Plugins/Experimental/MediaViewer/Source/MediaViewer/Private/ImageViewers/LevelEditorViewportImageViewer.cpp
--- a/Engine/Plugins/Experimental/MediaViewer/Source/MediaViewer/Private/ImageViewers/LevelEditorViewportImageViewer.cpp
+++ b/Engine/Plugins/Experimental/MediaViewer/Source/MediaViewer/Private/ImageViewers/LevelEditorViewportImageViewer.cpp
@@ -8,6 +8,7 @@
 #include "Modules/ModuleManager.h"
 #include "Slate/SceneViewport.h"
 #include "SLevelViewport.h"
+#include "Widgets/StrongSlateImageBrush.h"
 
 #define LOCTEXT_NAMESPACE "LevelEditorViewportImageViewer"
 
@@ -142,7 +143,7 @@ FText FLevelEditorViewportImageViewer::FItem::GetItemTypeDisplayName() const
 	return LOCTEXT("LevelEditorViewport", "Level Editor Viewport");
 }
 
-TSharedPtr<FSlateBrush> FLevelEditorViewportImageViewer::FItem::CreateThumbnail()
+UTextureRenderTarget2D* FLevelEditorViewportImageViewer::FItem::CreateRenderTargetThumbnail()
 {
 	if (StringValue.IsEmpty())
 	{
@@ -157,16 +158,28 @@ TSharedPtr<FSlateBrush> FLevelEditorViewportImageViewer::FItem::CreateThumbnail(
 	}
 
 	const FIntPoint RenderTargetSize = {32, 32};
-	RenderTarget = FMediaViewerUtils::CreateRenderTarget(RenderTargetSize, /* Transparent */ false);
+	UTextureRenderTarget2D* RenderTarget = FMediaViewerUtils::CreateRenderTarget(RenderTargetSize, /* Transparent */ false);
 
 	RenderViewport(
-		TNotNull<FViewport*>(Viewport.Get()), 
-		TNotNull<UTextureRenderTarget2D*>(RenderTarget), 
+		TNotNull<FViewport*>(Viewport.Get()),
+		TNotNull<UTextureRenderTarget2D*>(RenderTarget),
 		FRenderComplete(),
 		/* Resize render target */ false
 	);
 
-	TSharedRef<FSlateImageBrush> ThumbnailBrush = MakeShared<FSlateImageBrush>(RenderTarget, FVector2D(RenderTarget->GetSurfaceWidth(), RenderTarget->GetSurfaceDepth()));
+	return RenderTarget;
+}
+
+TSharedPtr<FSlateBrush> FLevelEditorViewportImageViewer::FItem::CreateThumbnail()
+{
+	UTextureRenderTarget2D* RenderTarget = CreateRenderTargetThumbnail();
+
+	if (!RenderTarget)
+	{
+		return nullptr;
+	}
+
+	TSharedRef<FSlateImageBrush> ThumbnailBrush = MakeShared<FStrongSlateImageBrush>(RenderTarget, FVector2D(RenderTarget->GetSurfaceWidth(), RenderTarget->GetSurfaceDepth()));
 
 	return ThumbnailBrush;
 }
@@ -196,20 +209,6 @@ TSharedPtr<FMediaViewerLibraryItem> FLevelEditorViewportImageViewer::FItem::Clon
 	return MakeShared<FLevelEditorViewportImageViewer::FItem>(FPrivateToken(), *this);
 }
 
-FString FLevelEditorViewportImageViewer::FItem::GetReferencerName() const
-{
-	static const FString ReferencerName = TEXT("FLevelEditorViewportImageViewer");
-	return ReferencerName;
-}
-
-void FLevelEditorViewportImageViewer::FItem::AddReferencedObjects(FReferenceCollector& Collector)
-{
-	if (RenderTarget)
-	{
-		Collector.AddReferencedObject(RenderTarget);
-	}
-}
-
 FLevelEditorViewportImageViewer::FLevelEditorViewportImageViewer(const FString& InViewportValue)
 	: FLevelEditorViewportImageViewer(FItem::GetIdForViewport(InViewportValue, /* Create id if invalid */ true), InViewportValue)
 {
diff --git a/Engine/Plugins/Experimental/MediaViewer/Source/MediaViewer/Private/ImageViewers/LevelEditorViewportImageViewer.h b/Engine/Plugins/Experimental/MediaViewer/Source/MediaViewer/Private/ImageViewers/LevelEditorViewportImageViewer.h
--- a/Engine/Plugins/Experimental/MediaViewer/Source/MediaViewer/Private/ImageViewers/LevelEditorViewportImageViewer.h
+++ b/Engine/Plugins/Experimental/MediaViewer/Source/MediaViewer/Private/ImageViewers/LevelEditorViewportImageViewer.h
@@ -32,14 +32,12 @@ public:
 		virtual TSharedPtr<FMediaViewerLibraryItem> CreateLibraryItem(const FMediaViewerLibraryItem& InSavedItem) const override;
 	};
 
-	struct FItem : public FMediaViewerLibraryItem, public FGCObject
+	struct FItem : public FMediaViewerLibraryItem
 	{
 		static FGuid GetIdForViewport(const FString& InConfigKey, bool bCreateIfInvalid);
 		static TSharedPtr<FSceneViewport> GetViewportFromConfigKey(const FString& InConfigKey);
 		static TSharedPtr<FSceneViewport> GetViewportFromStringValue(const FString& InStringValue);
 
-		TObjectPtr<UTextureRenderTarget2D> RenderTarget;
-
 		FItem(const FString& InConfigKey);
 
 		FItem(const FGuid& InId, const FString& InConfigKey);
@@ -56,10 +54,8 @@ public:
 		virtual TSharedPtr<FMediaViewerLibraryItem> Clone() const override;
 		//~ End FMediaViewerLibraryItem
 
-		//~ Begin FGCObject
-		virtual FString GetReferencerName() const override;
-		virtual void AddReferencedObjects(FReferenceCollector& Collector) override;
-		//~ End FGCObject
+	protected:
+		UTextureRenderTarget2D* CreateRenderTargetThumbnail();
 	};
 
 	static const FLazyName ItemTypeName;
diff --git a/Engine/Plugins/Experimental/MediaViewer/Source/MediaViewer/Private/ImageViewers/MaterialInterfaceImageViewer.cpp b/Engine/Plugins/Experimental/MediaViewer/Source/MediaViewer/Private/ImageViewers/MaterialInterfaceImageViewer.cpp
--- a/Engine/Plugins/Experimental/MediaViewer/Source/MediaViewer/Private/ImageViewers/MaterialInterfaceImageViewer.cpp
+++ b/Engine/Plugins/Experimental/MediaViewer/Source/MediaViewer/Private/ImageViewers/MaterialInterfaceImageViewer.cpp
@@ -2,7 +2,6 @@
 
 #include "ImageViewers/MaterialInterfaceImageViewer.h"
 
-#include "AssetRegistry/AssetData.h"
 #include "AssetRegistry/AssetRegistryModule.h"
 #include "AssetRegistry/IAssetRegistry.h"
 #include "Brushes/SlateImageBrush.h"
@@ -16,6 +15,7 @@
 #include "MediaViewerUtils.h"
 #include "Misc/Optional.h"
 #include "Misc/TVariant.h"
+#include "Widgets/StrongSlateImageBrush.h"
 
 #define LOCTEXT_NAMESPACE "MaterialInterfaceImageViewer"
 
@@ -152,21 +152,21 @@ FText FMaterialInterfaceImageViewer::FItem::GetItemTypeDisplayName() const
 	return LOCTEXT("Material", "Material");
 }
 
-void FMaterialInterfaceImageViewer::FItem::CreateRenderTargetThumbnail()
+UTextureRenderTarget2D* FMaterialInterfaceImageViewer::FItem::CreateRenderTargetThumbnail()
 {
 	if (StringValue.IsEmpty())
 	{
-		return;
+		return nullptr;
 	}
 
 	UMaterialInterface* MaterialInterface = LoadAssetFromString<UMaterialInterface>(StringValue);
 
 	if (!MaterialInterface)
 	{
-		return;
+		return nullptr;
 	}
 
-	RenderTarget = FMediaViewerUtils::RenderMaterial(MaterialInterface);
+	return FMediaViewerUtils::RenderMaterial(MaterialInterface);
 }
 
 TSharedPtr<FMediaViewerLibraryItem> FMaterialInterfaceImageViewer::FItem::Clone() const
@@ -179,6 +179,19 @@ TSharedPtr<FMediaViewerLibraryItem> FMaterialInterfaceImageViewer::FItem::Clone(
 	return MakeShared<FMaterialInterfaceImageViewer::FItem>(FPrivateToken(), *this);
 }
 
+TOptional<FAssetData> FMaterialInterfaceImageViewer::FItem::AsAsset() const
+{
+	if (UMaterialInterface* MaterialInterface = LoadAssetFromString<UMaterialInterface>(StringValue))
+	{
+		if (MaterialInterface->IsAsset())
+		{
+			return FAssetData(MaterialInterface);
+		}
+	}
+
+	return {};
+}
+
 TSharedPtr<FSlateBrush> FMaterialInterfaceImageViewer::FItem::CreateThumbnail()
 {
 	if (StringValue.IsEmpty())
@@ -186,12 +199,14 @@ TSharedPtr<FSlateBrush> FMaterialInterfaceImageViewer::FItem::CreateThumbnail()
 		return nullptr;
 	}
 
+	UTextureRenderTarget2D* RenderTarget = CreateRenderTargetThumbnail();
+
 	if (!RenderTarget)
 	{
-		CreateRenderTargetThumbnail();
+		return nullptr;
 	}
 
-	TSharedRef<FSlateImageBrush> ThumbnailBrush = MakeShared<FSlateImageBrush>(RenderTarget, FVector2D(RenderTarget->GetSurfaceWidth(), RenderTarget->GetSurfaceDepth()));
+	TSharedRef<FSlateImageBrush> ThumbnailBrush = MakeShared<FStrongSlateImageBrush>(RenderTarget, FVector2D(RenderTarget->GetSurfaceWidth(), RenderTarget->GetSurfaceDepth()));
 
 	return ThumbnailBrush;
 }
@@ -218,20 +233,6 @@ TSharedPtr<FMediaImageViewer> FMaterialInterfaceImageViewer::FItem::CreateImageV
 	return MakeShared<FMaterialInterfaceImageViewer>(MaterialInterface);
 }
 
-FString FMaterialInterfaceImageViewer::FItem::GetReferencerName() const
-{
-	static const FString ReferencerName = TEXT("FMaterialInterfaceImageViewer::FItem");
-	return ReferencerName;
-}
-
-void FMaterialInterfaceImageViewer::FItem::AddReferencedObjects(FReferenceCollector& Collector)
-{
-	if (RenderTarget)
-	{
-		Collector.AddReferencedObject(RenderTarget);
-	}
-}
-
 FMaterialInterfaceImageViewer::FMaterialInterfaceImageViewer(TNotNull<UMaterialInterface*> InMaterialInterface)
 	: FMaterialInterfaceImageViewer(FGuid::NewGuid(), InMaterialInterface)
 {
diff --git a/Engine/Plugins/Experimental/MediaViewer/Source/MediaViewer/Private/ImageViewers/MaterialInterfaceImageViewer.h b/Engine/Plugins/Experimental/MediaViewer/Source/MediaViewer/Private/ImageViewers/MaterialInterfaceImageViewer.h
--- a/Engine/Plugins/Experimental/MediaViewer/Source/MediaViewer/Private/ImageViewers/MaterialInterfaceImageViewer.h
+++ b/Engine/Plugins/Experimental/MediaViewer/Source/MediaViewer/Private/ImageViewers/MaterialInterfaceImageViewer.h
@@ -60,10 +60,8 @@ public:
 		virtual TSharedPtr<FMediaViewerLibraryItem> CreateLibraryItem(const FMediaViewerLibraryItem& InSavedItem) const override;
 	};
 
-	struct FItem : public FMediaViewerLibraryItem, public FGCObject
+	struct FItem : public FMediaViewerLibraryItem
 	{
-		TObjectPtr<UTextureRenderTarget2D> RenderTarget;
-
 		FItem(const FText& InName, const FText& InToolTip, bool bInTransient, TNotNull<UMaterialInterface*> InMaterial);
 
 		FItem(const FGuid& InId, const FText& InName, const FText& InToolTip, bool bInTransient, TNotNull<UMaterialInterface*> InMaterial);
@@ -72,20 +70,17 @@ public:
 
 		virtual ~FItem() override = default;
 
-		void CreateRenderTargetThumbnail();
-
 		//~ Begin FMediaViewerLibraryItem
 		virtual FName GetItemType() const override;
 		virtual FText GetItemTypeDisplayName() const override;
 		virtual TSharedPtr<FSlateBrush> CreateThumbnail() override;
 		virtual TSharedPtr<FMediaImageViewer> CreateImageViewer() const override;
 		virtual TSharedPtr<FMediaViewerLibraryItem> Clone() const override;
+		virtual TOptional<FAssetData> AsAsset() const override;
 		//~ End FMediaViewerLibraryItem
 
-		//~ Begin FGCObject
-		virtual FString GetReferencerName() const override;
-		virtual void AddReferencedObjects(FReferenceCollector& Collector) override;
-		//~ End FGCObject
+	protected:
+		UTextureRenderTarget2D* CreateRenderTargetThumbnail();
 	};
 
 	static const FLazyName ItemTypeName;
diff --git a/Engine/Plugins/Experimental/MediaViewer/Source/MediaViewer/Private/ImageViewers/MediaSourceImageViewer.cpp b/Engine/Plugins/Experimental/MediaViewer/Source/MediaViewer/Private/ImageViewers/MediaSourceImageViewer.cpp
--- a/Engine/Plugins/Experimental/MediaViewer/Source/MediaViewer/Private/ImageViewers/MediaSourceImageViewer.cpp
+++ b/Engine/Plugins/Experimental/MediaViewer/Source/MediaViewer/Private/ImageViewers/MediaSourceImageViewer.cpp
@@ -2,7 +2,6 @@
 
 #include "ImageViewers/MediaSourceImageViewer.h"
 
-#include "AssetRegistry/AssetData.h"
 #include "Brushes/SlateImageBrush.h"
 #include "DetailLayoutBuilder.h"
 #include "Engine/Texture2D.h"
@@ -202,6 +201,19 @@ TSharedPtr<FMediaViewerLibraryItem> FMediaSourceImageViewer::FAssetItem::Clone()
 	return MakeShared<FMediaSourceImageViewer::FAssetItem>(FPrivateToken(), *this);
 }
 
+TOptional<FAssetData> FMediaSourceImageViewer::FAssetItem::AsAsset() const
+{
+	if (UMediaSource* MediaSource = LoadAssetFromString<UMediaSource>(StringValue))
+	{
+		if (MediaSource->IsAsset())
+		{
+			return FAssetData(MediaSource);
+		}
+	}
+
+	return {};
+}
+
 FMediaSourceImageViewer::FExternalItem::FExternalItem(const FText& InName, const FText& InToolTip, const FString& InFilePath)
 	: FMediaSourceImageViewer::FExternalItem(FGuid::NewGuid(), InName, InToolTip, InFilePath)
 {
diff --git a/Engine/Plugins/Experimental/MediaViewer/Source/MediaViewer/Private/ImageViewers/MediaSourceImageViewer.h b/Engine/Plugins/Experimental/MediaViewer/Source/MediaViewer/Private/ImageViewers/MediaSourceImageViewer.h
--- a/Engine/Plugins/Experimental/MediaViewer/Source/MediaViewer/Private/ImageViewers/MediaSourceImageViewer.h
+++ b/Engine/Plugins/Experimental/MediaViewer/Source/MediaViewer/Private/ImageViewers/MediaSourceImageViewer.h
@@ -86,6 +86,7 @@ public:
 		virtual FText GetItemTypeDisplayName() const override;
 		virtual TSharedPtr<FMediaImageViewer> CreateImageViewer() const override;
 		virtual TSharedPtr<FMediaViewerLibraryItem> Clone() const override;
+		virtual TOptional<FAssetData> AsAsset() const override;
 		//~ End FMediaViewerLibraryItem
 	};
 
diff --git a/Engine/Plugins/Experimental/MediaViewer/Source/MediaViewer/Private/ImageViewers/MediaTextureImageViewer.cpp b/Engine/Plugins/Experimental/MediaViewer/Source/MediaViewer/Private/ImageViewers/MediaTextureImageViewer.cpp
--- a/Engine/Plugins/Experimental/MediaViewer/Source/MediaViewer/Private/ImageViewers/MediaTextureImageViewer.cpp
+++ b/Engine/Plugins/Experimental/MediaViewer/Source/MediaViewer/Private/ImageViewers/MediaTextureImageViewer.cpp
@@ -2,7 +2,6 @@
 
 #include "ImageViewers/MediaTextureImageViewer.h"
 
-#include "AssetRegistry/AssetData.h"
 #include "Brushes/SlateImageBrush.h"
 #include "Engine/TextureRenderTarget2D.h"
 #include "Materials/MaterialInstanceDynamic.h"
@@ -157,6 +156,19 @@ TSharedPtr<FMediaViewerLibraryItem> FMediaTextureImageViewer::FItem::Clone() con
 	return MakeShared<FMediaTextureImageViewer::FItem>(FPrivateToken(), *this);
 }
 
+TOptional<FAssetData> FMediaTextureImageViewer::FItem::AsAsset() const
+{
+	if (UMediaTexture* MediaTexture = LoadAssetFromString<UMediaTexture>(StringValue))
+	{
+		if (MediaTexture->IsAsset())
+		{
+			return FAssetData(MediaTexture);
+		}
+	}
+
+	return {};
+}
+
 FMediaTextureImageViewer::FMediaTextureImageViewer(TNotNull<UMediaTexture*> InMediaTexture)
 	: FMediaTextureImageViewer(FGuid::NewGuid(), InMediaTexture)
 {
diff --git a/Engine/Plugins/Experimental/MediaViewer/Source/MediaViewer/Private/ImageViewers/MediaTextureImageViewer.h b/Engine/Plugins/Experimental/MediaViewer/Source/MediaViewer/Private/ImageViewers/MediaTextureImageViewer.h
--- a/Engine/Plugins/Experimental/MediaViewer/Source/MediaViewer/Private/ImageViewers/MediaTextureImageViewer.h
+++ b/Engine/Plugins/Experimental/MediaViewer/Source/MediaViewer/Private/ImageViewers/MediaTextureImageViewer.h
@@ -67,6 +67,7 @@ public:
 		virtual TSharedPtr<FSlateBrush> CreateThumbnail() override;
 		virtual TSharedPtr<FMediaImageViewer> CreateImageViewer() const override;
 		virtual TSharedPtr<FMediaViewerLibraryItem> Clone() const override;
+		virtual TOptional<FAssetData> AsAsset() const override;
 		//~ End FMediaViewerLibraryItem
 	};
 
diff --git a/Engine/Plugins/Experimental/MediaViewer/Source/MediaViewer/Private/ImageViewers/Texture2DImageViewer.cpp b/Engine/Plugins/Experimental/MediaViewer/Source/MediaViewer/Private/ImageViewers/Texture2DImageViewer.cpp
--- a/Engine/Plugins/Experimental/MediaViewer/Source/MediaViewer/Private/ImageViewers/Texture2DImageViewer.cpp
+++ b/Engine/Plugins/Experimental/MediaViewer/Source/MediaViewer/Private/ImageViewers/Texture2DImageViewer.cpp
@@ -2,7 +2,6 @@
 
 #include "ImageViewers/Texture2DImageViewer.h"
 
-#include "AssetRegistry/AssetData.h"
 #include "Brushes/SlateImageBrush.h"
 #include "Engine/Texture.h"
 #include "Engine/Texture2D.h"
@@ -173,6 +172,19 @@ TSharedPtr<FMediaViewerLibraryItem> FTexture2DImageViewer::FItem::Clone() const
 	return MakeShared<FTexture2DImageViewer::FItem>(FPrivateToken(), *this);
 }
 
+TOptional<FAssetData> FTexture2DImageViewer::FItem::AsAsset() const
+{
+	if (::IsValid(Texture))
+	{
+		if (Texture->IsAsset())
+		{
+			return FAssetData(Texture);
+		}
+	}
+
+	return {};
+}
+
 FString FTexture2DImageViewer::FItem::GetReferencerName() const
 {
 	static const FString ReferencerName = TEXT("FTexture2DImageViewer::FItem");
diff --git a/Engine/Plugins/Experimental/MediaViewer/Source/MediaViewer/Private/ImageViewers/Texture2DImageViewer.h b/Engine/Plugins/Experimental/MediaViewer/Source/MediaViewer/Private/ImageViewers/Texture2DImageViewer.h
--- a/Engine/Plugins/Experimental/MediaViewer/Source/MediaViewer/Private/ImageViewers/Texture2DImageViewer.h
+++ b/Engine/Plugins/Experimental/MediaViewer/Source/MediaViewer/Private/ImageViewers/Texture2DImageViewer.h
@@ -68,6 +68,7 @@ public:
 		virtual TSharedPtr<FSlateBrush> CreateThumbnail() override;
 		virtual TSharedPtr<FMediaImageViewer> CreateImageViewer() const override;
 		virtual TSharedPtr<FMediaViewerLibraryItem> Clone() const override;
+		virtual TOptional<FAssetData> AsAsset() const override;
 		//~ End FMediaViewerLibraryItem
 
 		//~ Begin FGCObject
diff --git a/Engine/Plugins/Experimental/MediaViewer/Source/MediaViewer/Private/ImageViewers/TextureRenderTarget2DImageViewer.cpp b/Engine/Plugins/Experimental/MediaViewer/Source/MediaViewer/Private/ImageViewers/TextureRenderTarget2DImageViewer.cpp
--- a/Engine/Plugins/Experimental/MediaViewer/Source/MediaViewer/Private/ImageViewers/TextureRenderTarget2DImageViewer.cpp
+++ b/Engine/Plugins/Experimental/MediaViewer/Source/MediaViewer/Private/ImageViewers/TextureRenderTarget2DImageViewer.cpp
@@ -2,7 +2,6 @@
 
 #include "ImageViewers/TextureRenderTarget2DImageViewer.h"
 
-#include "AssetRegistry/AssetData.h"
 #include "Brushes/SlateImageBrush.h"
 #include "Engine/TextureRenderTarget2D.h"
 #include "Math/Vector2D.h"
@@ -154,6 +153,19 @@ TSharedPtr<FMediaViewerLibraryItem> FTextureRenderTarget2DImageViewer::FItem::Cl
 	return MakeShared<FTextureRenderTarget2DImageViewer::FItem>(FPrivateToken(), *this);
 }
 
+TOptional<FAssetData> FTextureRenderTarget2DImageViewer::FItem::AsAsset() const
+{
+	if (UTextureRenderTarget2D* RenderTarget = LoadAssetFromString<UTextureRenderTarget2D>(StringValue))
+	{
+		if (RenderTarget->IsAsset())
+		{
+			return FAssetData(RenderTarget);
+		}
+	}
+
+	return {};
+}
+
 FTextureRenderTarget2DImageViewer::FTextureRenderTarget2DImageViewer(TNotNull<UTextureRenderTarget2D*> InRenderTarget)
 	: FTextureRenderTarget2DImageViewer(FGuid::NewGuid(), InRenderTarget)
 {
diff --git a/Engine/Plugins/Experimental/MediaViewer/Source/MediaViewer/Private/ImageViewers/TextureRenderTarget2DImageViewer.h b/Engine/Plugins/Experimental/MediaViewer/Source/MediaViewer/Private/ImageViewers/TextureRenderTarget2DImageViewer.h
--- a/Engine/Plugins/Experimental/MediaViewer/Source/MediaViewer/Private/ImageViewers/TextureRenderTarget2DImageViewer.h
+++ b/Engine/Plugins/Experimental/MediaViewer/Source/MediaViewer/Private/ImageViewers/TextureRenderTarget2DImageViewer.h
@@ -64,6 +64,7 @@ public:
 		virtual TSharedPtr<FSlateBrush> CreateThumbnail() override;
 		virtual TSharedPtr<FMediaImageViewer> CreateImageViewer() const override;
 		virtual TSharedPtr<FMediaViewerLibraryItem> Clone() const override;
+		virtual TOptional<FAssetData> AsAsset() const override;
 		//~ End FMediaViewerLibraryItem
 	};
 
diff --git a/Engine/Plugins/Experimental/MediaViewer/Source/MediaViewer/Private/Widgets/MediaViewerLibraryItemAssetDragDropOperation.cpp b/Engine/Plugins/Experimental/MediaViewer/Source/MediaViewer/Private/Widgets/MediaViewerLibraryItemAssetDragDropOperation.cpp
new file mode 100644
--- /dev/null
+++ b/Engine/Plugins/Experimental/MediaViewer/Source/MediaViewer/Private/Widgets/MediaViewerLibraryItemAssetDragDropOperation.cpp
@@ -0,0 +1,109 @@
+// Copyright Epic Games, Inc. All Rights Reserved.
+
+#include "Widgets/MediaViewerLibraryItemAssetDragDropOperation.h"
+
+#include "DetailLayoutBuilder.h"
+#include "Library/MediaViewerLibrary.h"
+#include "Library/MediaViewerLibraryGroup.h"
+#include "Library/MediaViewerLibraryItem.h"
+#include "Math/Color.h"
+#include "MediaViewerLibraryItemAssetDragDropOperation.h"
+#include "Styling/AppStyle.h"
+#include "Widgets/Images/SImage.h"
+#include "Widgets/Layout/SBorder.h"
+#include "Widgets/SBoxPanel.h"
+#include "Widgets/Text/STextBlock.h"
+#include "Widgets/Views/STableRow.h"
+
+#define LOCTEXT_NAMESPACE "MediaViewerLibraryItemAssetDragDropOperation"
+
+namespace UE::MediaViewer::Private
+{
+
+FMediaViewerLibraryItemAssetDragDropOperation::FMediaViewerLibraryItemAssetDragDropOperation(const TSharedRef<FMediaViewerLibrary>& InLibrary,
+	const IMediaViewerLibrary::FGroupItem& InGroupItem)
+	: GroupItem(InGroupItem)
+	, Decorator(SNew(SBorder))
+{
+	if (TSharedPtr<FMediaViewerLibraryItem> Item = InLibrary->GetItem(InGroupItem.ItemId))
+	{
+		if (TOptional<FAssetData> ItemAssetData = Item->AsAsset())
+		{
+			// Initializes the base asset drop operation
+			Init({ItemAssetData.GetValue()}, {}, nullptr);
+		}
+
+		CreateDecorator(Item);
+	}
+	
+	Construct();
+}
+
+const IMediaViewerLibrary::FGroupItem& FMediaViewerLibraryItemAssetDragDropOperation::GetGroupItem() const
+{
+	return GroupItem;
+}
+
+TSharedPtr<SWidget> FMediaViewerLibraryItemAssetDragDropOperation::GetDefaultDecorator() const
+{
+	return Decorator;
+}
+
+FCursorReply FMediaViewerLibraryItemAssetDragDropOperation::OnCursorQuery()
+{
+	return FCursorReply::Cursor(EMouseCursor::GrabHandClosed);
+}
+
+void FMediaViewerLibraryItemAssetDragDropOperation::CreateDecorator(const TSharedPtr<FMediaViewerLibraryItem>& InItem)
+{
+	if (!InItem.IsValid())
+	{
+		return;
+	}
+
+	ThumbnailBrush = InItem->CreateThumbnail();
+
+	Decorator->SetContent(
+		SNew(SHorizontalBox)
+		+ SHorizontalBox::Slot()
+		.AutoWidth()
+		.Padding(5.f, 5.f, 5.f, 5.f)
+		.VAlign(VAlign_Center)
+		[
+			SNew(SImage)
+			.Image(ThumbnailBrush.IsValid() ? ThumbnailBrush.Get() : nullptr)
+			.DesiredSizeOverride(FVector2D(24.f))
+		] 
+		+ SHorizontalBox::Slot()
+		.FillWidth(1.f)
+		.Padding(3.f, 5.f, 5.f, 5.f)
+		.HAlign(HAlign_Fill)
+		.VAlign(VAlign_Center)
+		[
+			SNew(SVerticalBox)
+			+ SVerticalBox::Slot()
+			.AutoHeight()
+			.HAlign(HAlign_Left)
+			.VAlign(VAlign_Center)
+			[
+				SNew(STextBlock)
+				.Text(InItem->Name)
+				.ToolTipText(InItem->ToolTip)
+				.Font(IDetailLayoutBuilder::GetDetailFont())
+			]
+			+ SVerticalBox::Slot()
+			.AutoHeight()
+			.HAlign(HAlign_Left)
+			.VAlign(VAlign_Center)
+			[
+				SNew(STextBlock)
+				.Text(InItem->GetItemTypeDisplayName())
+				.Font(IDetailLayoutBuilder::GetDetailFont())
+			]
+		]
+	);
+}
+
+} // UE::MediaViewer::Private
+
+#undef LOCTEXT_NAMESPACE
diff --git a/Engine/Plugins/Experimental/MediaViewer/Source/MediaViewer/Private/Widgets/MediaViewerLibraryItemAssetDragDropOperation.h b/Engine/Plugins/Experimental/MediaViewer/Source/MediaViewer/Private/Widgets/MediaViewerLibraryItemAssetDragDropOperation.h
new file mode 100644
--- /dev/null
+++ b/Engine/Plugins/Experimental/MediaViewer/Source/MediaViewer/Private/Widgets/MediaViewerLibraryItemAssetDragDropOperation.h
@@ -0,0 +1,52 @@
+// Copyright Epic Games, Inc. All Rights Reserved.
+
+#pragma once
+
+#include "DragAndDrop/AssetDragDropOp.h"
+
+#include "Library/IMediaViewerLibrary.h"
+#include "Templates/SharedPointer.h"
+
+class SBorder;
+enum class EItemDropZone;
+struct FMediaViewerLibraryEntry;
+struct FMediaViewerLibraryGroup;
+struct FMediaViewerLibraryItem;
+struct FSlateBrush;
+
+namespace UE::MediaViewer
+{
+	enum class EMediaImageViewerPosition : uint8;
+}
+
+namespace UE::MediaViewer::Private
+{
+
+class FMediaViewerLibrary;
+
+class FMediaViewerLibraryItemAssetDragDropOperation : public FAssetDragDropOp
+{
+public:
+	DRAG_DROP_OPERATOR_TYPE(FMediaViewerLibraryItemAssetDragDropOperation, FAssetDragDropOp)
+
+	FMediaViewerLibraryItemAssetDragDropOperation(const TSharedRef<FMediaViewerLibrary>& InLibrary,
+		const IMediaViewerLibrary::FGroupItem& InGroupItem);
+	
+	virtual ~FMediaViewerLibraryItemAssetDragDropOperation() override = default;
+
+	const IMediaViewerLibrary::FGroupItem& GetGroupItem() const;
+
+	//~ Begin FDragDropOp
+	virtual TSharedPtr<SWidget> GetDefaultDecorator() const override;
+	virtual FCursorReply OnCursorQuery() override;
+	//~ End FDragDropOp
+
+protected:
+	IMediaViewerLibrary::FGroupItem GroupItem;
+	TSharedPtr<FSlateBrush> ThumbnailBrush;
+	TSharedRef<SBorder> Decorator;
+
+	void CreateDecorator(const TSharedPtr<FMediaViewerLibraryItem>& InItem);
+};
+
+} // UE::MediaViewer::Private
diff --git a/Engine/Plugins/Experimental/MediaViewer/Source/MediaViewer/Private/Widgets/MediaViewerLibraryItemDragDropOperation.cpp b/Engine/Plugins/Experimental/MediaViewer/Source/MediaViewer/Private/Widgets/MediaViewerLibraryItemDragDropOperation.cpp
--- a/Engine/Plugins/Experimental/MediaViewer/Source/MediaViewer/Private/Widgets/MediaViewerLibraryItemDragDropOperation.cpp
+++ b/Engine/Plugins/Experimental/MediaViewer/Source/MediaViewer/Private/Widgets/MediaViewerLibraryItemDragDropOperation.cpp
@@ -53,14 +53,6 @@ void FMediaViewerLibraryItemDragDropOperation::CreateDecorator(const TSharedPtr<
 
 	ThumbnailBrush = InItem->CreateThumbnail();
 
-	const FLinearColor SaveColor = InItem->IsTransient()
-		? FLinearColor(0.2, 0.2, 0.2)
-		: FLinearColor::White;
-
-	const FText SaveToolTip = InItem->IsTransient()
-		? LOCTEXT("TransientToolTip", "Cannot be saved between sessions.")
-		: LOCTEXT("NotTransientToolTip", "Can be saved between sessions.");
-
 	Decorator->SetContent(
 		SNew(SHorizontalBox)
 		+ SHorizontalBox::Slot()
diff --git a/Engine/Plugins/Experimental/MediaViewer/Source/MediaViewer/Private/Widgets/SMediaViewerLibraryItem.cpp b/Engine/Plugins/Experimental/MediaViewer/Source/MediaViewer/Private/Widgets/SMediaViewerLibraryItem.cpp
--- a/Engine/Plugins/Experimental/MediaViewer/Source/MediaViewer/Private/Widgets/SMediaViewerLibraryItem.cpp
+++ b/Engine/Plugins/Experimental/MediaViewer/Source/MediaViewer/Private/Widgets/SMediaViewerLibraryItem.cpp
@@ -20,6 +20,7 @@
 #include "Widgets/Input/SButton.h"
 #include "Widgets/Input/SCheckBox.h"
 #include "Widgets/MediaViewerDelegates.h"
+#include "Widgets/MediaViewerLibraryItemAssetDragDropOperation.h"
 #include "Widgets/MediaViewerLibraryItemDragDropOperation.h"
 #include "Widgets/SBoxPanel.h"
 #include "Widgets/SNullWidget.h"
@@ -292,6 +293,11 @@ FReply SMediaViewerLibraryItem::OnRemoveButtonClicked()
 
 FReply SMediaViewerLibraryItem::OnDragDetected(const FGeometry& InMyGeometry, const FPointerEvent& InPointerEvent)
 {
+	if (!Item.IsValid())
+	{
+		return FReply::Handled();
+	}
+
 	TSharedPtr<FMediaViewerLibrary> Library = LibraryWeak.Pin();
 
 	if (!Library.IsValid())
@@ -299,6 +305,14 @@ FReply SMediaViewerLibraryItem::OnDragDetected(const FGeometry& InMyGeometry, co
 		return FReply::Handled();
 	}
 
+	if (Item->AsAsset())
+	{
+		return FReply::Handled().BeginDragDrop(MakeShared<FMediaViewerLibraryItemAssetDragDropOperation>(
+			Library.ToSharedRef(),
+			GroupItem
+		));
+	}
+
 	return FReply::Handled().BeginDragDrop(MakeShared<FMediaViewerLibraryItemDragDropOperation>(
 		Library.ToSharedRef(),
 		GroupItem
@@ -535,9 +549,14 @@ void SMediaViewerLibraryItem::OnAssetsDropped(TConstArrayView<FAssetData> InAsse
 
 const FSlateBrush* SMediaViewerLibraryItem::GetThumbnail() const
 {
-	if (ThumbnailBrush.IsValid() && IsValid(ThumbnailBrush->GetResourceObject()))
+	if (ThumbnailBrush.IsValid())
 	{
-		return ThumbnailBrush.Get();
+		UObject* ResourceObject = ThumbnailBrush->GetResourceObject();
+
+		if (ResourceObject && ResourceObject->IsValidLowLevel() && IsValid(ResourceObject))
+		{
+			return ThumbnailBrush.Get();
+		}
 	}
 
 	return FAppStyle::Get().GetBrush("SourceControl.StatusIcon.Unknown");
diff --git a/Engine/Plugins/Experimental/MediaViewer/Source/MediaViewer/Private/Widgets/StrongSlateImageBrush.cpp b/Engine/Plugins/Experimental/MediaViewer/Source/MediaViewer/Private/Widgets/StrongSlateImageBrush.cpp
new file mode 100644
--- /dev/null
+++ b/Engine/Plugins/Experimental/MediaViewer/Source/MediaViewer/Private/Widgets/StrongSlateImageBrush.cpp
@@ -0,0 +1,24 @@
+// Copyright Epic Games, Inc. All Rights Reserved.
+
+#include "Widgets/StrongSlateImageBrush.h"
+
+FStrongSlateImageBrush::FStrongSlateImageBrush(UObject* InResourceObject, const UE::Slate::FDeprecateVector2DParameter& InImageSize, 
+	const FSlateColor& InTint, ESlateBrushTileType::Type InTiling, ESlateBrushImageType::Type InImageType)
+	: FSlateImageBrush(InResourceObject, InImageSize, InTint, InTiling, InImageType)
+	, ResourceObjectPtr(InResourceObject)
+{
+}
+
+FString FStrongSlateImageBrush::GetReferencerName() const
+{
+	static const FString ReferencerName = TEXT("StrongSlateImageBrush");
+	return ReferencerName;
+}
+
+void FStrongSlateImageBrush::AddReferencedObjects(FReferenceCollector& InCollector)
+{
+	if (ResourceObjectPtr)
+	{
+		InCollector.AddReferencedObject(ResourceObjectPtr);
+	}
+}
diff --git a/Engine/Plugins/Experimental/MediaViewer/Source/MediaViewer/Private/Widgets/StrongSlateImageBrush.h b/Engine/Plugins/Experimental/MediaViewer/Source/MediaViewer/Private/Widgets/StrongSlateImageBrush.h
new file mode 100644
--- /dev/null
+++ b/Engine/Plugins/Experimental/MediaViewer/Source/MediaViewer/Private/Widgets/StrongSlateImageBrush.h
@@ -0,0 +1,25 @@
+// Copyright Epic Games, Inc. All Rights Reserved.
+
+#pragma once
+
+#include "Brushes/SlateImageBrush.h"
+#include "UObject/GCObject.h"
+
+/**
+ * Ignores the Margin. Just renders the image. Can tile the image instead of stretching.
+ */
+struct FStrongSlateImageBrush : public FSlateImageBrush, public FGCObject
+{
+	FStrongSlateImageBrush(UObject* InResourceObject, const UE::Slate::FDeprecateVector2DParameter& InImageSize, 
+		const FSlateColor& InTint = FSlateColor(FLinearColor(1, 1, 1, 1)), 
+		ESlateBrushTileType::Type InTiling = ESlateBrushTileType::NoTile, 
+		ESlateBrushImageType::Type InImageType = ESlateBrushImageType::FullColor);
+
+	//~ Begin FGCObject
+	virtual FString GetReferencerName() const override;
+	virtual void AddReferencedObjects(FReferenceCollector& InCollector) override;
+	//~ End FGCObject
+
+protected:
+	TObjectPtr<UObject> ResourceObjectPtr;
+};
diff --git a/Engine/Plugins/Experimental/MediaViewer/Source/MediaViewer/Public/Library/MediaViewerLibraryItem.h b/Engine/Plugins/Experimental/MediaViewer/Source/MediaViewer/Public/Library/MediaViewerLibraryItem.h
--- a/Engine/Plugins/Experimental/MediaViewer/Source/MediaViewer/Public/Library/MediaViewerLibraryItem.h
+++ b/Engine/Plugins/Experimental/MediaViewer/Source/MediaViewer/Public/Library/MediaViewerLibraryItem.h
@@ -4,10 +4,12 @@
 
 #include "Library/MediaViewerLibraryEntry.h"
 
+#include "AssetRegistry/AssetData.h"
 #include "Containers/UnrealString.h"
 #include "HAL/Platform.h"
 #include "Internationalization/Text.h"
 #include "Misc/Guid.h"
+#include "Misc/Optional.h"
 #include "UObject/ObjectPtr.h"
 
 #include "MediaViewerLibraryItem.generated.h"
@@ -67,6 +69,8 @@ public:
 
 	MEDIAVIEWER_API const FString& GetStringValue() const;
 
+	virtual TOptional<FAssetData> AsAsset() const { return {}; }
+
 	//~ Begin FMediaViewerLibraryEntry
 	virtual UE::MediaViewer::EMediaViewerLibraryEntryType GetEntryType() const override { return UE::MediaViewer::EMediaViewerLibraryEntryType::Item; }
 	//~ End FMediaViewerLibraryEntry
