Media Viewer: Fixed pixel sampling for compressed images
- All pixel samples now work via a render target

diff --git a/Engine/Plugins/Experimental/MediaViewer/Source/MediaViewer/Private/ImageViewers/TextureSampleCache.cpp b/Engine/Plugins/Experimental/MediaViewer/Source/MediaViewer/Private/ImageViewers/TextureSampleCache.cpp
--- a/Engine/Plugins/Experimental/MediaViewer/Source/MediaViewer/Private/ImageViewers/TextureSampleCache.cpp
+++ b/Engine/Plugins/Experimental/MediaViewer/Source/MediaViewer/Private/ImageViewers/TextureSampleCache.cpp
@@ -3,94 +3,30 @@
 #include "ImageViewers/TextureSampleCache.h"
 
 #include "Engine/Texture.h"
-#include "ImageViewers/TextureMipCache.h"
+#include "Engine/TextureRenderTarget2D.h"
+#include "GlobalShader.h"
+#include "MediaViewerUtils.h"
 #include "Misc/ScopeLock.h"
 #include "PixelFormat.h"
+#include "RenderGraphBuilder.h"
+#include "RenderGraphUtils.h"
 #include "RenderingThread.h"
 #include "TextureResource.h"
 #include "UObject/StrongObjectPtr.h"
+#include "UObject/WeakObjectPtr.h"
 
 namespace UE::MediaViewer::Private
 {
 
-bool CanCopyTextureWithRHICommand(EPixelFormat InFormat)
-{
-	switch (InFormat)
-	{
-		// Textures encoding that cannot be read by the RHI.
-		// Based on the list of invalid texture formats in the Metal RHI (see MetalRHI.cpp)
-		case EPixelFormat::PF_DXT1:
-		case EPixelFormat::PF_DXT3:
-		case EPixelFormat::PF_DXT5:
-		case EPixelFormat::PF_BC4:
-		case EPixelFormat::PF_BC5:
-		case EPixelFormat::PF_BC6H:
-		case EPixelFormat::PF_BC7:
-		case EPixelFormat::PF_DepthStencil:
-		case EPixelFormat::PF_ShadowDepth:
-		case EPixelFormat::PF_D24:
-		case EPixelFormat::PF_A1:
-		case EPixelFormat::PF_PVRTC2:
-		case EPixelFormat::PF_PVRTC4:
-		case EPixelFormat::PF_R5G6B5_UNORM:
-		case EPixelFormat::PF_B5G5R5A1_UNORM:
-		case EPixelFormat::PF_ATC_RGB:
-		case EPixelFormat::PF_ATC_RGBA_E:
-		case EPixelFormat::PF_ATC_RGBA_I:
-		case EPixelFormat::PF_X24_G8:
-		case EPixelFormat::PF_ETC1:
-		case EPixelFormat::PF_ETC2_RGB:
-		case EPixelFormat::PF_ETC2_RGBA:
-		case EPixelFormat::PF_ASTC_4x4:
-		case EPixelFormat::PF_ASTC_6x6:
-		case EPixelFormat::PF_ASTC_8x8:
-		case EPixelFormat::PF_ASTC_10x10:
-		case EPixelFormat::PF_ASTC_12x12:
-		case EPixelFormat::PF_ASTC_4x4_HDR:
-		case EPixelFormat::PF_ASTC_6x6_HDR:
-		case EPixelFormat::PF_ASTC_8x8_HDR:
-		case EPixelFormat::PF_ASTC_10x10_HDR:
-		case EPixelFormat::PF_ASTC_12x12_HDR:
-		case EPixelFormat::PF_L8:
-		case EPixelFormat::PF_R16G16B16A16_SNORM:
-		case EPixelFormat::PF_PLATFORM_HDR_0:
-		case EPixelFormat::PF_PLATFORM_HDR_1:
-		case EPixelFormat::PF_NV12:
-		case EPixelFormat::PF_ETC2_R11_EAC:
-		case EPixelFormat::PF_ETC2_RG11_EAC:
-		case EPixelFormat::PF_R32G32B32_UINT:
-		case EPixelFormat::PF_R32G32B32_SINT:
-		case EPixelFormat::PF_R32G32B32F:
-		case EPixelFormat::PF_R64_UINT:
-		case EPixelFormat::PF_R9G9B9EXP5:
-		case EPixelFormat::PF_P010:
-		case EPixelFormat::PF_ASTC_4x4_NORM_RG:
-		case EPixelFormat::PF_ASTC_6x6_NORM_RG:
-		case EPixelFormat::PF_ASTC_8x8_NORM_RG:
-		case EPixelFormat::PF_ASTC_10x10_NORM_RG:
-		case EPixelFormat::PF_ASTC_12x12_NORM_RG:
-		case EPixelFormat::PF_R8G8B8:
-			return false;
-
-		default:
-			return true;
-	}
-}
-
 FTextureSampleCache::FTextureSampleCache()
-	: PixelFormat(PF_Unknown)
 {
 }
 
 FTextureSampleCache::FTextureSampleCache(TNotNull<UTexture*> InTexture, EPixelFormat InPixelFormat)
 	: Texture(InTexture)
-	, PixelFormat(InPixelFormat)
 	, SampleCS(FCriticalSection())
+	, RenderTarget(nullptr)
 {
-	if (!CanCopyTextureWithRHICommand(PixelFormat))
-	{
-		MipCache = MakeShared<FTextureMipCache>(InTexture);
-	}
 }
 
 bool FTextureSampleCache::IsValid() const
@@ -100,6 +36,18 @@ bool FTextureSampleCache::IsValid() const
 
 bool FTextureSampleCache::NeedsUpdate(const FIntPoint& InPixelCoordinates, const TOptional<FTimespan>& InTime) const
 {
+	if (!Texture)
+	{
+		return false;
+	}
+
+	constexpr int32 MinimumSurfaceSize = 3;
+
+	if (Texture->GetSurfaceWidth() < MinimumSurfaceSize || Texture->GetSurfaceHeight() < MinimumSurfaceSize)
+	{
+		return false;
+	}
+
 	FScopeLock Lock(&SampleCS);
 	return bDirty || !PixelColorSample.IsSet() || InPixelCoordinates != PixelColorSample->Coordinates
 		|| (InTime.IsSet() && (PixelColorSample->Time != InTime.GetValue()));
@@ -116,13 +64,14 @@ const FLinearColor* FTextureSampleCache::GetPixelColor(const FIntPoint& InPixelC
 
 	if (bNeedsUpdate)
 	{
-		if (!CanCopyTextureWithRHICommand(PixelFormat))
+		// If we're already using a render target, we're good to directly sample it.
+		if (Texture && Texture->IsA<UTextureRenderTarget2D>())
 		{
-			SetPixelColor_Mip(InPixelCoordinates, InTime);
+			SetPixelColor_RHI(InPixelCoordinates, InTime, Texture);
 		}
 		else
 		{
-			SetPixelColor_RHI(InPixelCoordinates, InTime);
+			SetPixelColor_RenderTarget(InPixelCoordinates, InTime);
 		}
 	}
 
@@ -150,11 +99,6 @@ void FTextureSampleCache::Invalidate()
 		FScopeLock Lock(&SampleCS);
 		PixelColorSample.Reset();
 	}
-
-	if (MipCache.IsValid())
-	{
-		MipCache->Invalidate();
-	}
 }
 
 FTextureSampleCache& FTextureSampleCache::operator=(const FTextureSampleCache& InOther)
@@ -163,28 +107,47 @@ FTextureSampleCache& FTextureSampleCache::operator=(const FTextureSampleCache& I
 	FScopeLock OtherLock(&InOther.SampleCS);
 
 	Texture = InOther.Texture;
-	PixelFormat = InOther.PixelFormat;
 	bDirty = InOther.bDirty;
 	PixelColorSample = InOther.PixelColorSample;
-	MipCache = InOther.MipCache;
 	
 	return *this;
 }
 
-void FTextureSampleCache::SetPixelColor_RHI(const FIntPoint& InPixelCoordinates, const TOptional<FTimespan>& InTime)
+FString FTextureSampleCache::GetReferencerName() const
+{
+	static const FString ReferencerName = TEXT("FTextureSampleCache");
+	return ReferencerName;
+}
+
+void FTextureSampleCache::AddReferencedObjects(FReferenceCollector& Collector)
+{
+	if (RenderTarget)
+	{
+		Collector.AddReferencedObject(RenderTarget);
+	}
+}
+
+void FTextureSampleCache::SetPixelColor_RHI(const FIntPoint& InPixelCoordinates, const TOptional<FTimespan>& InTime, UTexture* InTexture)
 {
 	ENQUEUE_RENDER_COMMAND(GetPixelColors)
 	(
-		[ThisWeak = SharedThis(this).ToWeakPtr(), InPixelCoordinates, InTime](FRHICommandListImmediate& RHICmdList)
+		[ThisWeak = SharedThis(this).ToWeakPtr(), TextureWeak = TWeakObjectPtr<UTexture>(InTexture), InPixelCoordinates, InTime](FRHICommandListImmediate& RHICmdList)
 		{
 			TSharedPtr<FTextureSampleCache> This = ThisWeak.Pin();
 
-			if (!This.IsValid() || !This->IsValid())
+			if (!This.IsValid())
 			{
 				return;
 			}
 
-			if (FTextureResource* TextureResource = This->Texture->GetResource())
+			TStrongObjectPtr<UTexture> TextureToSample = TextureWeak.Pin();
+
+			if (!TextureToSample.IsValid())
+			{
+				return;
+			}
+
+			if (FTextureResource* TextureResource = TextureToSample->GetResource())
 			{
 				const FIntVector Size = TextureResource->TextureRHI->GetSizeXYZ();
 				TArray<FLinearColor> Data;
@@ -216,33 +179,75 @@ void FTextureSampleCache::SetPixelColor_RHI(const FIntPoint& InPixelCoordinates,
 	);
 }
 
-void FTextureSampleCache::SetPixelColor_Mip(const FIntPoint& InPixelCoordinates, const TOptional<FTimespan>& InTime)
+void FTextureSampleCache::SetPixelColor_RenderTarget(const FIntPoint& InPixelCoordinates, const TOptional<FTimespan>& InTime)
 {
-	if (!MipCache.IsValid())
+	if (!RenderTarget)
 	{
-		return;
+		RenderTarget = FMediaViewerUtils::CreateRenderTarget(
+			FIntPoint(
+				FMath::RoundToInt(Texture->GetSurfaceWidth()),
+				FMath::RoundToInt(Texture->GetSurfaceHeight())
+			),
+			/* Transparent */ true
+		);
 	}
 
-	const FImage* Mip = MipCache->GetMipImage(/* Mip level */ 0);
+	bool bNeedsUpdate = false;
 
-	if (!Mip)
 	{
 		FScopeLock Lock(&SampleCS);
-		PixelColorSample.Reset();
-		return;
+		bNeedsUpdate = !PixelColorSample.IsSet() || InTime != PixelColorSample->Time;
 	}
 
+	if (!bNeedsUpdate)
 	{
-		FScopeLock Lock(&SampleCS);
+		SetPixelColor_RHI(InPixelCoordinates, InTime, RenderTarget);
+		return;
+	}
 
-		PixelColorSample = {
-			InPixelCoordinates,
-			InTime,
-			Mip->GetOnePixelLinear(InPixelCoordinates.X, InPixelCoordinates.Y)
-		};
+	FRenderComplete OnComplete = FRenderComplete::CreateSPLambda(
+		this,
+		[this, InPixelCoordinates, InTime](bool bSuccess)
+		{
+			if (bSuccess)
+			{
+				SetPixelColor_RHI(InPixelCoordinates, InTime, RenderTarget);
+			}
+		}
+	);
 
-		bDirty = false;
-	}
+	ENQUEUE_RENDER_COMMAND(CopyViewportRenderTarget)(
+		[Source = Texture, Target = RenderTarget, OnComplete](FRHICommandListImmediate& InRHICmdList)
+		{
+			FTextureResource* SourceResource = Source->GetResource();
+
+			if (!SourceResource)
+			{
+				OnComplete.ExecuteIfBound(false);
+				return;
+			}
+
+			FTextureResource* TargetResource = Target->GetResource();
+
+			if (!TargetResource)
+			{
+				OnComplete.ExecuteIfBound(false);
+				return;
+			}
+
+			FRDGBuilder GraphBuilder(InRHICmdList);
+			
+			FRDGTexture* InputTexture = GraphBuilder.RegisterExternalTexture(CreateRenderTarget(SourceResource->TextureRHI, TEXT("SourceTexture")));
+			FRDGTexture* OutputTexture = GraphBuilder.RegisterExternalTexture(CreateRenderTarget(TargetResource->TextureRHI, TEXT("TargetTexture")));
+			const FGlobalShaderMap* GlobalShaderMap = GetGlobalShaderMap(GMaxRHIFeatureLevel);
+
+			AddDrawTexturePass(GraphBuilder, GlobalShaderMap, InputTexture, OutputTexture, FRDGDrawTextureInfo());
+
+			GraphBuilder.Execute();
+
+			OnComplete.ExecuteIfBound(true);
+		}
+	);
 }
 
 } // UE::MediaViewer::Private
diff --git a/Engine/Plugins/Experimental/MediaViewer/Source/MediaViewer/Private/ImageViewers/TextureSampleCache.h b/Engine/Plugins/Experimental/MediaViewer/Source/MediaViewer/Private/ImageViewers/TextureSampleCache.h
--- a/Engine/Plugins/Experimental/MediaViewer/Source/MediaViewer/Private/ImageViewers/TextureSampleCache.h
+++ b/Engine/Plugins/Experimental/MediaViewer/Source/MediaViewer/Private/ImageViewers/TextureSampleCache.h
@@ -3,6 +3,7 @@
 #pragma once
 
 #include "Templates/SharedPointer.h"
+#include "UObject/GCObject.h"
 
 #include "HAL/CriticalSection.h"
 #include "Math/Color.h"
@@ -12,14 +13,13 @@
 #include "Misc/Timespan.h"
 
 class UTexture;
+class UTextureRenderTarget2D;
 enum EPixelFormat : uint8;
 
 namespace UE::MediaViewer::Private
 {
 
-class FTextureMipCache;
-
-class FTextureSampleCache : public TSharedFromThis<FTextureSampleCache>
+class FTextureSampleCache : public TSharedFromThis<FTextureSampleCache>, public FGCObject
 {
 public:
 	FTextureSampleCache();
@@ -35,6 +35,11 @@ public:
 
 	FTextureSampleCache& operator=(const FTextureSampleCache& InOther);
 
+	//~ Begin FGCObject
+	virtual FString GetReferencerName() const override;
+	virtual void AddReferencedObjects(FReferenceCollector& Collector) override;
+	//~ End FGCObject
+
 protected:
 	struct FPixelColorSample
 	{
@@ -44,17 +49,16 @@ protected:
 	};
 
 	UTexture* Texture = nullptr;
-	EPixelFormat PixelFormat;
 	bool bDirty = true;
 	TOptional<FPixelColorSample> PixelColorSample;
-	TSharedPtr<FTextureMipCache> MipCache;
 	mutable FCriticalSection SampleCS;
+	TObjectPtr<UTextureRenderTarget2D> RenderTarget;
 
-	void SetPixelColor_RHI(const FIntPoint& InPixelCoordinates, const TOptional<FTimespan>& InTime);
+	bool NeedsUpdate(const FIntPoint& InPixelCoordinates, const TOptional<FTimespan>& InTime) const;
 
-	void SetPixelColor_Mip(const FIntPoint& InPixelCoordinates, const TOptional<FTimespan>& InTime);
+	void SetPixelColor_RHI(const FIntPoint& InPixelCoordinates, const TOptional<FTimespan>& InTime, UTexture* InTexture);
 
-	bool NeedsUpdate(const FIntPoint& InPixelCoordinates, const TOptional<FTimespan>& InTime) const;
+	void SetPixelColor_RenderTarget(const FIntPoint& InPixelCoordinates, const TOptional<FTimespan>& InTime);
 };
 
 } // UE::MediaViewer::Private
diff --git a/Engine/Plugins/Experimental/MediaViewer/Source/MediaViewer/Private/ImageViewers/ViewportImageViewer.h b/Engine/Plugins/Experimental/MediaViewer/Source/MediaViewer/Private/ImageViewers/ViewportImageViewer.h
--- a/Engine/Plugins/Experimental/MediaViewer/Source/MediaViewer/Private/ImageViewers/ViewportImageViewer.h
+++ b/Engine/Plugins/Experimental/MediaViewer/Source/MediaViewer/Private/ImageViewers/ViewportImageViewer.h
@@ -8,6 +8,7 @@
 
 #include "Delegates/IDelegateInstance.h"
 #include "ImageViewers/TextureSampleCache.h"
+#include "MediaViewerUtils.h"
 #include "Templates/SharedPointer.h"
 
 #include "ViewportImageViewer.generated.h"
@@ -58,7 +59,6 @@ protected:
 		Invalidated
 	};
 
-	DECLARE_DELEGATE_OneParam(FRenderComplete, bool /* Was succsesfully rendered */)
 	static void RenderViewport(TNotNull<FViewport*> InViewport, TNotNull<UTextureRenderTarget2D*> InRenderTarget, 
 		FRenderComplete InRenderComplete, bool bResizeTargetToViewport = true);
 
diff --git a/Engine/Plugins/Experimental/MediaViewer/Source/MediaViewer/Private/MediaViewerUtils.h b/Engine/Plugins/Experimental/MediaViewer/Source/MediaViewer/Private/MediaViewerUtils.h
--- a/Engine/Plugins/Experimental/MediaViewer/Source/MediaViewer/Private/MediaViewerUtils.h
+++ b/Engine/Plugins/Experimental/MediaViewer/Source/MediaViewer/Private/MediaViewerUtils.h
@@ -37,6 +37,8 @@ class UMediaViewerUserData : public UAssetUserData
 namespace UE::MediaViewer::Private
 {
 
+DECLARE_DELEGATE_OneParam(FRenderComplete, bool /* Was succsesfully rendered */)
+
 class FMediaViewerUtils
 {
 public:
