Media Viewer: Fixed issue with A-image not being displayed behind B-image when slider was set to 0%.

diff --git a/Engine/Plugins/Experimental/MediaViewer/Source/MediaViewer/Private/Widgets/SMediaViewer.cpp b/Engine/Plugins/Experimental/MediaViewer/Source/MediaViewer/Private/Widgets/SMediaViewer.cpp
--- a/Engine/Plugins/Experimental/MediaViewer/Source/MediaViewer/Private/Widgets/SMediaViewer.cpp
+++ b/Engine/Plugins/Experimental/MediaViewer/Source/MediaViewer/Private/Widgets/SMediaViewer.cpp
@@ -449,12 +449,9 @@ int32 SMediaViewer::OnPaint(const FPaintArgs& InArgs, const FGeometry& InAllotte
 			break;
 
 		case EMediaImageViewerActivePosition::Both:
-			if (!FMath::IsNearlyZero(ABSplitterLocation))
-			{
-				UVRange.SetLowerBound(0.f);
-				UVRange.SetUpperBound(1.f);
-				GetImageViewer(EMediaImageViewerPosition::First)->Paint(PaintParams);
-			}
+			UVRange.SetLowerBound(0.f);
+			UVRange.SetUpperBound(1.f);
+			GetImageViewer(EMediaImageViewerPosition::First)->Paint(PaintParams);
 
 			if (!FMath::IsNearlyZero(ABSplitterLocation - 1.f))
 			{
