Media Viewer: Updated to more closely match Figma design

diff --git a/Engine/Plugins/Experimental/MediaViewer/Source/MediaViewer/Private/ImageViewers/LevelEditorViewportImageViewer.cpp b/Engine/Plugins/Experimental/MediaViewer/Source/MediaViewer/Private/ImageViewers/LevelEditorViewportImageViewer.cpp
--- a/Engine/Plugins/Experimental/MediaViewer/Source/MediaViewer/Private/ImageViewers/LevelEditorViewportImageViewer.cpp
+++ b/Engine/Plugins/Experimental/MediaViewer/Source/MediaViewer/Private/ImageViewers/LevelEditorViewportImageViewer.cpp
@@ -3,6 +3,7 @@
 #include "ImageViewers/LevelEditorViewportImageViewer.h"
 
 #include "Engine/TextureRenderTarget2D.h"
+#include "Engine/World.h"
 #include "LevelEditor.h"
 #include "MediaViewerUtils.h"
 #include "Modules/ModuleManager.h"
@@ -138,6 +139,11 @@ FName FLevelEditorViewportImageViewer::FItem::GetItemType() const
 	return FLevelEditorViewportImageViewer::ItemTypeName;
 }
 
+FSlateColor FLevelEditorViewportImageViewer::FItem::GetItemTypeColor() const
+{
+	return GetClassColor(UWorld::StaticClass());
+}
+
 FText FLevelEditorViewportImageViewer::FItem::GetItemTypeDisplayName() const
 {
 	return LOCTEXT("LevelEditorViewport", "Level Editor Viewport");
diff --git a/Engine/Plugins/Experimental/MediaViewer/Source/MediaViewer/Private/ImageViewers/LevelEditorViewportImageViewer.h b/Engine/Plugins/Experimental/MediaViewer/Source/MediaViewer/Private/ImageViewers/LevelEditorViewportImageViewer.h
--- a/Engine/Plugins/Experimental/MediaViewer/Source/MediaViewer/Private/ImageViewers/LevelEditorViewportImageViewer.h
+++ b/Engine/Plugins/Experimental/MediaViewer/Source/MediaViewer/Private/ImageViewers/LevelEditorViewportImageViewer.h
@@ -49,6 +49,7 @@ public:
 		//~ Begin FMediaViewerLibraryItem
 		virtual FName GetItemType() const override;
 		virtual FText GetItemTypeDisplayName() const override;
+		virtual FSlateColor GetItemTypeColor() const override;
 		virtual TSharedPtr<FSlateBrush> CreateThumbnail() override;
 		virtual TSharedPtr<FMediaImageViewer> CreateImageViewer() const override;
 		virtual TSharedPtr<FMediaViewerLibraryItem> Clone() const override;
diff --git a/Engine/Plugins/Experimental/MediaViewer/Source/MediaViewer/Private/ImageViewers/MaterialInterfaceImageViewer.cpp b/Engine/Plugins/Experimental/MediaViewer/Source/MediaViewer/Private/ImageViewers/MaterialInterfaceImageViewer.cpp
--- a/Engine/Plugins/Experimental/MediaViewer/Source/MediaViewer/Private/ImageViewers/MaterialInterfaceImageViewer.cpp
+++ b/Engine/Plugins/Experimental/MediaViewer/Source/MediaViewer/Private/ImageViewers/MaterialInterfaceImageViewer.cpp
@@ -147,6 +147,11 @@ FName FMaterialInterfaceImageViewer::FItem::GetItemType() const
 	return FMaterialInterfaceImageViewer::ItemTypeName;
 }
 
+FSlateColor FMaterialInterfaceImageViewer::FItem::GetItemTypeColor() const
+{
+	return GetClassColor(UMaterial::StaticClass());
+}
+
 FText FMaterialInterfaceImageViewer::FItem::GetItemTypeDisplayName() const
 {
 	return LOCTEXT("Material", "Material");
diff --git a/Engine/Plugins/Experimental/MediaViewer/Source/MediaViewer/Private/ImageViewers/MaterialInterfaceImageViewer.h b/Engine/Plugins/Experimental/MediaViewer/Source/MediaViewer/Private/ImageViewers/MaterialInterfaceImageViewer.h
--- a/Engine/Plugins/Experimental/MediaViewer/Source/MediaViewer/Private/ImageViewers/MaterialInterfaceImageViewer.h
+++ b/Engine/Plugins/Experimental/MediaViewer/Source/MediaViewer/Private/ImageViewers/MaterialInterfaceImageViewer.h
@@ -73,6 +73,7 @@ public:
 		//~ Begin FMediaViewerLibraryItem
 		virtual FName GetItemType() const override;
 		virtual FText GetItemTypeDisplayName() const override;
+		virtual FSlateColor GetItemTypeColor() const override;
 		virtual TSharedPtr<FSlateBrush> CreateThumbnail() override;
 		virtual TSharedPtr<FMediaImageViewer> CreateImageViewer() const override;
 		virtual TSharedPtr<FMediaViewerLibraryItem> Clone() const override;
diff --git a/Engine/Plugins/Experimental/MediaViewer/Source/MediaViewer/Private/ImageViewers/MediaSourceImageViewer.cpp b/Engine/Plugins/Experimental/MediaViewer/Source/MediaViewer/Private/ImageViewers/MediaSourceImageViewer.cpp
--- a/Engine/Plugins/Experimental/MediaViewer/Source/MediaViewer/Private/ImageViewers/MediaSourceImageViewer.cpp
+++ b/Engine/Plugins/Experimental/MediaViewer/Source/MediaViewer/Private/ImageViewers/MediaSourceImageViewer.cpp
@@ -167,6 +167,11 @@ FName FMediaSourceImageViewer::FAssetItem::GetItemType() const
 	return FMediaSourceImageViewer::ItemTypeName_Asset;
 }
 
+FSlateColor FMediaSourceImageViewer::FItem::GetItemTypeColor() const
+{
+	return GetClassColor(UFileMediaSource::StaticClass());
+}
+
 FText FMediaSourceImageViewer::FAssetItem::GetItemTypeDisplayName() const
 {
 	return LOCTEXT("MediaSource", "Media Asset");
diff --git a/Engine/Plugins/Experimental/MediaViewer/Source/MediaViewer/Private/ImageViewers/MediaSourceImageViewer.h b/Engine/Plugins/Experimental/MediaViewer/Source/MediaViewer/Private/ImageViewers/MediaSourceImageViewer.h
--- a/Engine/Plugins/Experimental/MediaViewer/Source/MediaViewer/Private/ImageViewers/MediaSourceImageViewer.h
+++ b/Engine/Plugins/Experimental/MediaViewer/Source/MediaViewer/Private/ImageViewers/MediaSourceImageViewer.h
@@ -62,6 +62,7 @@ public:
 		FItem(const FGuid& InId, const FText& InName, const FText& InToolTip, bool bInTransient, const FString& InStringValue);
 
 		//~ Begin FMediaViewerLibraryItem
+		virtual FSlateColor GetItemTypeColor() const override;
 		virtual TSharedPtr<FSlateBrush> CreateThumbnail() override;
 		//~ End FMediaViewerLibraryItem
 
diff --git a/Engine/Plugins/Experimental/MediaViewer/Source/MediaViewer/Private/ImageViewers/MediaTextureImageViewer.cpp b/Engine/Plugins/Experimental/MediaViewer/Source/MediaViewer/Private/ImageViewers/MediaTextureImageViewer.cpp
--- a/Engine/Plugins/Experimental/MediaViewer/Source/MediaViewer/Private/ImageViewers/MediaTextureImageViewer.cpp
+++ b/Engine/Plugins/Experimental/MediaViewer/Source/MediaViewer/Private/ImageViewers/MediaTextureImageViewer.cpp
@@ -115,6 +115,11 @@ FText FMediaTextureImageViewer::FItem::GetItemTypeDisplayName() const
 	return LOCTEXT("MediaTexture", "Media Texture");
 }
 
+FSlateColor FMediaTextureImageViewer::FItem::GetItemTypeColor() const
+{
+	return GetClassColor(UMediaTexture::StaticClass());
+}
+
 TSharedPtr<FSlateBrush> FMediaTextureImageViewer::FItem::CreateThumbnail()
 {
 	UMediaTexture* MediaTexture = LoadAssetFromString<UMediaTexture>(StringValue);
diff --git a/Engine/Plugins/Experimental/MediaViewer/Source/MediaViewer/Private/ImageViewers/MediaTextureImageViewer.h b/Engine/Plugins/Experimental/MediaViewer/Source/MediaViewer/Private/ImageViewers/MediaTextureImageViewer.h
--- a/Engine/Plugins/Experimental/MediaViewer/Source/MediaViewer/Private/ImageViewers/MediaTextureImageViewer.h
+++ b/Engine/Plugins/Experimental/MediaViewer/Source/MediaViewer/Private/ImageViewers/MediaTextureImageViewer.h
@@ -64,6 +64,7 @@ public:
 		//~ Begin FMediaViewerLibraryItem
 		virtual FName GetItemType() const override;
 		virtual FText GetItemTypeDisplayName() const override;
+		virtual FSlateColor GetItemTypeColor() const override;
 		virtual TSharedPtr<FSlateBrush> CreateThumbnail() override;
 		virtual TSharedPtr<FMediaImageViewer> CreateImageViewer() const override;
 		virtual TSharedPtr<FMediaViewerLibraryItem> Clone() const override;
diff --git a/Engine/Plugins/Experimental/MediaViewer/Source/MediaViewer/Private/ImageViewers/Texture2DImageViewer.cpp b/Engine/Plugins/Experimental/MediaViewer/Source/MediaViewer/Private/ImageViewers/Texture2DImageViewer.cpp
--- a/Engine/Plugins/Experimental/MediaViewer/Source/MediaViewer/Private/ImageViewers/Texture2DImageViewer.cpp
+++ b/Engine/Plugins/Experimental/MediaViewer/Source/MediaViewer/Private/ImageViewers/Texture2DImageViewer.cpp
@@ -115,6 +115,11 @@ FText FTexture2DImageViewer::FItem::GetItemTypeDisplayName() const
 	return LOCTEXT("Texture2D", "Texture 2D");
 }
 
+FSlateColor FTexture2DImageViewer::FItem::GetItemTypeColor() const
+{
+	return GetClassColor(UTexture2D::StaticClass());
+}
+
 TSharedPtr<FSlateBrush> FTexture2DImageViewer::FItem::CreateThumbnail()
 {
 	if (StringValue.IsEmpty())
diff --git a/Engine/Plugins/Experimental/MediaViewer/Source/MediaViewer/Private/ImageViewers/Texture2DImageViewer.h b/Engine/Plugins/Experimental/MediaViewer/Source/MediaViewer/Private/ImageViewers/Texture2DImageViewer.h
--- a/Engine/Plugins/Experimental/MediaViewer/Source/MediaViewer/Private/ImageViewers/Texture2DImageViewer.h
+++ b/Engine/Plugins/Experimental/MediaViewer/Source/MediaViewer/Private/ImageViewers/Texture2DImageViewer.h
@@ -65,6 +65,7 @@ public:
 		//~ Begin FMediaViewerLibraryItem
 		virtual FName GetItemType() const override;
 		virtual FText GetItemTypeDisplayName() const override;
+		virtual FSlateColor GetItemTypeColor() const override;
 		virtual TSharedPtr<FSlateBrush> CreateThumbnail() override;
 		virtual TSharedPtr<FMediaImageViewer> CreateImageViewer() const override;
 		virtual TSharedPtr<FMediaViewerLibraryItem> Clone() const override;
diff --git a/Engine/Plugins/Experimental/MediaViewer/Source/MediaViewer/Private/ImageViewers/TextureRenderTarget2DImageViewer.cpp b/Engine/Plugins/Experimental/MediaViewer/Source/MediaViewer/Private/ImageViewers/TextureRenderTarget2DImageViewer.cpp
--- a/Engine/Plugins/Experimental/MediaViewer/Source/MediaViewer/Private/ImageViewers/TextureRenderTarget2DImageViewer.cpp
+++ b/Engine/Plugins/Experimental/MediaViewer/Source/MediaViewer/Private/ImageViewers/TextureRenderTarget2DImageViewer.cpp
@@ -112,6 +112,11 @@ FText FTextureRenderTarget2DImageViewer::FItem::GetItemTypeDisplayName() const
 	return LOCTEXT("TextureRenderTarget2D", "Render Target 2D");
 }
 
+FSlateColor FTextureRenderTarget2DImageViewer::FItem::GetItemTypeColor() const
+{
+	return GetClassColor(UTextureRenderTarget2D::StaticClass());
+}
+
 TSharedPtr<FSlateBrush> FTextureRenderTarget2DImageViewer::FItem::CreateThumbnail()
 {
 	UTextureRenderTarget2D* RenderTarget = LoadAssetFromString<UTextureRenderTarget2D>(StringValue);
diff --git a/Engine/Plugins/Experimental/MediaViewer/Source/MediaViewer/Private/ImageViewers/TextureRenderTarget2DImageViewer.h b/Engine/Plugins/Experimental/MediaViewer/Source/MediaViewer/Private/ImageViewers/TextureRenderTarget2DImageViewer.h
--- a/Engine/Plugins/Experimental/MediaViewer/Source/MediaViewer/Private/ImageViewers/TextureRenderTarget2DImageViewer.h
+++ b/Engine/Plugins/Experimental/MediaViewer/Source/MediaViewer/Private/ImageViewers/TextureRenderTarget2DImageViewer.h
@@ -61,6 +61,7 @@ public:
 		//~ Begin FMediaViewerLibraryItem
 		virtual FName GetItemType() const override;
 		virtual FText GetItemTypeDisplayName() const override;
+		virtual FSlateColor GetItemTypeColor() const override;
 		virtual TSharedPtr<FSlateBrush> CreateThumbnail() override;
 		virtual TSharedPtr<FMediaImageViewer> CreateImageViewer() const override;
 		virtual TSharedPtr<FMediaViewerLibraryItem> Clone() const override;
diff --git a/Engine/Plugins/Experimental/MediaViewer/Source/MediaViewer/Private/Library/MediaViewerLibraryItem.cpp b/Engine/Plugins/Experimental/MediaViewer/Source/MediaViewer/Private/Library/MediaViewerLibraryItem.cpp
--- a/Engine/Plugins/Experimental/MediaViewer/Source/MediaViewer/Private/Library/MediaViewerLibraryItem.cpp
+++ b/Engine/Plugins/Experimental/MediaViewer/Source/MediaViewer/Private/Library/MediaViewerLibraryItem.cpp
@@ -2,6 +2,11 @@
 
 #include "Library/MediaViewerLibraryItem.h"
 
+#include "AssetDefinitionRegistry.h"
+#include "AssetToolsModule.h"
+#include "Modules/ModuleManager.h"
+#include "Styling/StyleColors.h"
+
 #define LOCTEXT_NAMESPACE "MediaViewerLibraryItem"
 
 FMediaViewerLibraryItem::FMediaViewerLibraryItem()
@@ -33,6 +38,11 @@ FText FMediaViewerLibraryItem::GetItemTypeDisplayName() const
 	return LOCTEXT("Error", "Error");
 }
 
+FSlateColor FMediaViewerLibraryItem::GetItemTypeColor() const
+{
+	return GetClassColor(nullptr);
+}
+
 bool FMediaViewerLibraryItem::IsTransient() const
 {
 	return bTransient;
@@ -43,4 +53,20 @@ const FString& FMediaViewerLibraryItem::GetStringValue() const
 	return StringValue;
 }
 
+FSlateColor FMediaViewerLibraryItem::GetClassColor(UClass* InClass)
+{
+	if (InClass)
+	{
+		FAssetToolsModule& AssetToolsModule = FModuleManager::LoadModuleChecked<FAssetToolsModule>(TEXT("AssetTools"));
+		TWeakPtr<IAssetTypeActions> AssetTypeActionsWeak = AssetToolsModule.Get().GetAssetTypeActionsForClass(InClass);
+
+		if (TSharedPtr<IAssetTypeActions> AssetTypeActions = AssetTypeActionsWeak.Pin())
+		{
+			return AssetTypeActions->GetTypeColor();
+		}
+	}
+
+	return FSlateColor::UseForeground();
+}
+
 #undef LOCTEXT_NAMESPACE
diff --git a/Engine/Plugins/Experimental/MediaViewer/Source/MediaViewer/Private/MediaViewerStyle.cpp b/Engine/Plugins/Experimental/MediaViewer/Source/MediaViewer/Private/MediaViewerStyle.cpp
--- a/Engine/Plugins/Experimental/MediaViewer/Source/MediaViewer/Private/MediaViewerStyle.cpp
+++ b/Engine/Plugins/Experimental/MediaViewer/Source/MediaViewer/Private/MediaViewerStyle.cpp
@@ -3,8 +3,10 @@
 #include "MediaViewerStyle.h"
 
 #include "DetailLayoutBuilder.h"
+#include "Interfaces/IPluginManager.h"
 #include "Styling/AppStyle.h"
 #include "Styling/CoreStyle.h"
+#include "Styling/SlateStyleMacros.h"
 #include "Styling/SlateStyleRegistry.h"
 #include "Styling/SlateTypes.h"
 #include "Styling/StyleColors.h"
@@ -60,6 +62,12 @@ FMediaViewerStyle::FMediaViewerStyle()
 		.SetNormalPadding(FMargin(2.f, 0.f, 2.f, 0.f))
 		.SetPressedPadding(FMargin(2.f, 0.f, 2.f, 0.f)));
 
+	if (TSharedPtr<IPlugin> Plugin = IPluginManager::Get().FindPlugin(UE_PLUGIN_NAME))
+	{
+		SetContentRoot(FPaths::Combine(Plugin->GetBaseDir(), TEXT("Content")));
+		Set("ColorChip", new IMAGE_BRUSH_SVG(TEXT("color-chip"), FVector2D(12.0)));
+	}
+
 	FSlateStyleRegistry::RegisterSlateStyle(*this);
 }
 
diff --git a/Engine/Plugins/Experimental/MediaViewer/Source/MediaViewer/Private/Widgets/SMediaViewerLibraryItem.cpp b/Engine/Plugins/Experimental/MediaViewer/Source/MediaViewer/Private/Widgets/SMediaViewerLibraryItem.cpp
--- a/Engine/Plugins/Experimental/MediaViewer/Source/MediaViewer/Private/Widgets/SMediaViewerLibraryItem.cpp
+++ b/Engine/Plugins/Experimental/MediaViewer/Source/MediaViewer/Private/Widgets/SMediaViewerLibraryItem.cpp
@@ -24,6 +24,7 @@
 #include "Widgets/MediaViewerLibraryItemDragDropOperation.h"
 #include "Widgets/SBoxPanel.h"
 #include "Widgets/SNullWidget.h"
+#include "Widgets/SOverlay.h"
 #include "Widgets/Text/STextBlock.h"
 
 #define LOCTEXT_NAMESPACE "SMediaViewerLibraryItem"
@@ -71,12 +72,25 @@ void SMediaViewerLibraryItem::Construct(const FArguments& InArgs, const TSharedR
 	TSharedRef<SHorizontalBox> Box = SNew(SHorizontalBox)
 		+ SHorizontalBox::Slot()
 		.AutoWidth()
-		.Padding(0.f, 5.f, 0.f, 5.f)
+		.Padding(0.f, 5.f, 5.f, 5.f)
 		.VAlign(VAlign_Center)
 		[
-			SNew(SImage)
-			.Image(this, &SMediaViewerLibraryItem::GetThumbnail)
-			.DesiredSizeOverride(FVector2D(24.f))
+			SNew(SOverlay)
+			+ SOverlay::Slot()
+			[
+				SNew(SImage)
+				.Image(this, &SMediaViewerLibraryItem::GetThumbnail)
+				.DesiredSizeOverride(FVector2D(40.f))
+			]
+			+ SOverlay::Slot()
+			.Padding(30.f, 0.f, 0.f, 30.f)
+			[
+				SNew(SImage)
+				.Image(FMediaViewerStyle::Get().GetBrush("ColorChip"))
+				.DesiredSizeOverride(FVector2D(12.0))
+				.ColorAndOpacity(Item->GetItemTypeColor())
+				.RenderTransform(FSlateRenderTransform(FVector2D(2.f, -2.f)))
+			]
 		] 
 		+ SHorizontalBox::Slot()
 		.FillWidth(1.f)
@@ -89,6 +103,7 @@ void SMediaViewerLibraryItem::Construct(const FArguments& InArgs, const TSharedR
 			.AutoHeight()
 			.HAlign(HAlign_Left)
 			.VAlign(VAlign_Center)
+			.Padding(0.f, 0.f, 0.f, 5.f)
 			[
 				SNew(STextBlock)
 				.Text(Item->Name)
diff --git a/Engine/Plugins/Experimental/MediaViewer/Source/MediaViewer/Public/Library/MediaViewerLibraryItem.h b/Engine/Plugins/Experimental/MediaViewer/Source/MediaViewer/Public/Library/MediaViewerLibraryItem.h
--- a/Engine/Plugins/Experimental/MediaViewer/Source/MediaViewer/Public/Library/MediaViewerLibraryItem.h
+++ b/Engine/Plugins/Experimental/MediaViewer/Source/MediaViewer/Public/Library/MediaViewerLibraryItem.h
@@ -16,6 +16,7 @@
 
 class UTexture;
 struct FSlateBrush;
+struct FSlateColor;
 
 namespace UE::MediaViewer
 {
@@ -59,6 +60,8 @@ public:
 
 	MEDIAVIEWER_API virtual FText GetItemTypeDisplayName() const;
 
+	MEDIAVIEWER_API virtual FSlateColor GetItemTypeColor() const;
+
 	MEDIAVIEWER_API bool IsTransient() const;
 
 	virtual TSharedPtr<FSlateBrush> CreateThumbnail() { return nullptr; }
@@ -76,6 +79,8 @@ public:
 	//~ End FMediaViewerLibraryEntry
 
 protected:
+	MEDIAVIEWER_API static FSlateColor GetClassColor(UClass* InClass);
+
 	/** If all this entry and all its children are transient. */
 	bool bTransient;
 
