Media Viewer: Library visual update.
- Changed corner chip to color bar.
- Also added some padding around the thumbnail image, similar to the content browser.
- Fixed text clipping on item type.
- Added border to active entries.

diff --git a/Engine/Plugins/Experimental/MediaViewer/Source/MediaViewer/Private/Widgets/SMediaViewerLibraryItem.cpp b/Engine/Plugins/Experimental/MediaViewer/Source/MediaViewer/Private/Widgets/SMediaViewerLibraryItem.cpp
--- a/Engine/Plugins/Experimental/MediaViewer/Source/MediaViewer/Private/Widgets/SMediaViewerLibraryItem.cpp
+++ b/Engine/Plugins/Experimental/MediaViewer/Source/MediaViewer/Private/Widgets/SMediaViewerLibraryItem.cpp
@@ -57,39 +57,28 @@ void SMediaViewerLibraryItem::Construct(const FArguments& InArgs, const TSharedR
 
 	ThumbnailBrush = Item->CreateThumbnail();
 
-	TSharedRef<SHorizontalBox> Details = SNew(SHorizontalBox)
-		+ SHorizontalBox::Slot()
-		.AutoWidth()
-		.VAlign(VAlign_Center)
-		.Padding(0.f, 0.f, 5.f, 0.f)
-		[
-			SNew(STextBlock)
-			.Clipping(EWidgetClipping::ClipToBounds)
-			.Font(IDetailLayoutBuilder::GetDetailFont())
-			.Text(Item->GetItemTypeDisplayName())
-		];
-
 	TSharedRef<SHorizontalBox> Box = SNew(SHorizontalBox)
 		+ SHorizontalBox::Slot()
 		.AutoWidth()
 		.Padding(0.f, 5.f, 5.f, 5.f)
 		.VAlign(VAlign_Center)
 		[
-			SNew(SOverlay)
-			+ SOverlay::Slot()
+			SNew(SVerticalBox)
+			+ SVerticalBox::Slot()
+			.AutoHeight()
+			.Padding(2.f)
 			[
 				SNew(SImage)
 				.Image(this, &SMediaViewerLibraryItem::GetThumbnail)
-				.DesiredSizeOverride(FVector2D(40.f))
+				.DesiredSizeOverride(FVector2D(36.f))
 			]
-			+ SOverlay::Slot()
-			.Padding(30.f, 0.f, 0.f, 30.f)
+			+ SVerticalBox::Slot()
+			.AutoHeight()
 			[
-				SNew(SImage)
-				.Image(FMediaViewerStyle::Get().GetBrush("ColorChip"))
-				.DesiredSizeOverride(FVector2D(12.0))
-				.ColorAndOpacity(Item->GetItemTypeColor())
-				.RenderTransform(FSlateRenderTransform(FVector2D(2.f, -2.f)))
+				SNew(SColorBlock)
+				.Size(FVector2D(40.f, 2.f))
+				.Color(Item->GetItemTypeColor().GetSpecifiedColor())
+				.CornerRadius(FVector4(0.f, 0.f, 2.f, 2.f))
 			]
 		] 
 		+ SHorizontalBox::Slot()
@@ -99,6 +88,7 @@ void SMediaViewerLibraryItem::Construct(const FArguments& InArgs, const TSharedR
 		.VAlign(VAlign_Center)
 		[
 			SNew(SVerticalBox)
+			.Clipping(EWidgetClipping::ClipToBounds)
 			+ SVerticalBox::Slot()
 			.AutoHeight()
 			.HAlign(HAlign_Left)
@@ -115,7 +105,9 @@ void SMediaViewerLibraryItem::Construct(const FArguments& InArgs, const TSharedR
 			.HAlign(HAlign_Left)
 			.VAlign(VAlign_Center)
 			[
-				Details
+				SNew(STextBlock)
+				.Font(IDetailLayoutBuilder::GetDetailFont())
+				.Text(Item->GetItemTypeDisplayName())
 			]
 		];
 
@@ -183,13 +175,20 @@ void SMediaViewerLibraryItem::Construct(const FArguments& InArgs, const TSharedR
 
 	TSharedRef<SOverlay> Overlay = SNew(SOverlay)
 		+ SOverlay::Slot()
+		.Padding(0.f, 0.f, 2.f, 0.f)
+		[
+			SNew(SColorBlock)
+			.Color(this, &SMediaViewerLibraryItem::GetBorderColor)
+			.CornerRadius(2.f)
+		]
+		+ SOverlay::Slot()
+		.Padding(2.f, 2.f, 4.f, 2.f)
 		[
 			SNew(SColorBlock)
-			.Color(FStyleColors::Panel.GetSpecifiedColor())
-			.Visibility(this, &SMediaViewerLibraryItem::GetHoverVisibility)
+			.Color(this, &SMediaViewerLibraryItem::GetFillColor)
 		]
 		+ SOverlay::Slot()
-		.Padding(0.f, 0.f, 2.f, 0.f)
+		.Padding(6.f, 2.f, 4.f, 2.f)
 		[
 			Box
 		];
@@ -204,14 +203,27 @@ void SMediaViewerLibraryItem::Construct(const FArguments& InArgs, const TSharedR
 	);
 }
 
-EVisibility SMediaViewerLibraryItem::GetHoverVisibility() const
+FLinearColor SMediaViewerLibraryItem::GetBorderColor() const
+{
+	for (int32 Index = 0; Index < static_cast<int32>(EMediaImageViewerPosition::COUNT); ++Index)
+	{
+		if (IsActive(static_cast<EMediaImageViewerPosition>(Index)))
+		{
+			return FStyleColors::Primary.GetSpecifiedColor();
+		}
+	}
+
+	return GetFillColor();
+}
+
+FLinearColor SMediaViewerLibraryItem::GetFillColor() const
 {
 	if (IsHovered())
 	{
-		return EVisibility::Visible;
+		return FStyleColors::Panel.GetSpecifiedColor();
 	}
 
-	return EVisibility::Collapsed;
+	return FStyleColors::Background.GetSpecifiedColor();
 }
 
 bool SMediaViewerLibraryItem::IsActive(EMediaImageViewerPosition InPosition) const
diff --git a/Engine/Plugins/Experimental/MediaViewer/Source/MediaViewer/Private/Widgets/SMediaViewerLibraryItem.h b/Engine/Plugins/Experimental/MediaViewer/Source/MediaViewer/Private/Widgets/SMediaViewerLibraryItem.h
--- a/Engine/Plugins/Experimental/MediaViewer/Source/MediaViewer/Private/Widgets/SMediaViewerLibraryItem.h
+++ b/Engine/Plugins/Experimental/MediaViewer/Source/MediaViewer/Private/Widgets/SMediaViewerLibraryItem.h
@@ -57,7 +57,9 @@ protected:
 	TSharedPtr<FMediaViewerLibraryItem> Item;
 	TSharedPtr<FSlateBrush> ThumbnailBrush;
 
-	EVisibility GetHoverVisibility() const;
+	FLinearColor GetBorderColor() const;
+
+	FLinearColor GetFillColor() const;
 
 	bool IsActive(EMediaImageViewerPosition InPosition) const;
 
