Media Viewer: When in dual view and without an A image, the slider position works correctly.

diff --git a/Engine/Plugins/Experimental/MediaViewer/Source/MediaViewer/Private/Widgets/SMediaViewer.cpp b/Engine/Plugins/Experimental/MediaViewer/Source/MediaViewer/Private/Widgets/SMediaViewer.cpp
--- a/Engine/Plugins/Experimental/MediaViewer/Source/MediaViewer/Private/Widgets/SMediaViewer.cpp
+++ b/Engine/Plugins/Experimental/MediaViewer/Source/MediaViewer/Private/Widgets/SMediaViewer.cpp
@@ -678,14 +678,17 @@ float SMediaViewer::GetDPIScale() const
 
 void SMediaViewer::OnABResized(float InSize)
 {
-	if (!IsImageViewerNull(EMediaImageViewerPosition::First))
+	if (ActiveView != EMediaImageViewerActivePosition::Single)
 	{
-		StaticCastSharedRef<SSplitter>(ContentSlot->GetWidget())->SlotAt(0).SetSizeValue(InSize);
-	}
+		if (TSharedPtr<SMediaImageViewerOverlay> FirstOverlay = GetOverlay(EMediaImageViewerPosition::First))
+		{
+			StaticCastSharedRef<SSplitter>(ContentSlot->GetWidget())->SlotAt(0).SetSizeValue(InSize);
+		}
 
-	if (!IsImageViewerNull(EMediaImageViewerPosition::Second))
-	{
-		StaticCastSharedRef<SSplitter>(ContentSlot->GetWidget())->SlotAt(1).SetSizeValue(1.f - InSize);
+		if (TSharedPtr<SMediaImageViewerOverlay> FirstOverlay = GetOverlay(EMediaImageViewerPosition::Second))
+		{
+			StaticCastSharedRef<SSplitter>(ContentSlot->GetWidget())->SlotAt(1).SetSizeValue(1.f - InSize);
+		}
 	}
 
 	MediaViewerSettings.ABSplitterLocation = InSize * 100.f;
