Media Viewer: Snapshot button

diff --git a/Engine/Plugins/Experimental/MediaViewer/Source/MediaViewer/MediaViewer.build.cs b/Engine/Plugins/Experimental/MediaViewer/Source/MediaViewer/MediaViewer.build.cs
--- a/Engine/Plugins/Experimental/MediaViewer/Source/MediaViewer/MediaViewer.build.cs
+++ b/Engine/Plugins/Experimental/MediaViewer/Source/MediaViewer/MediaViewer.build.cs
@@ -22,6 +22,8 @@ public class MediaViewer : ModuleRules
 			{
 				"AppFramework",
 				"ApplicationCore",
+				"ContentBrowser",
+				"ContentBrowserData",
 				"EditorWidgets",
 				"Engine",
 				"ImageCore",
diff --git a/Engine/Plugins/Experimental/MediaViewer/Source/MediaViewer/Private/ImageViewers/LevelEditorViewportImageViewer.cpp b/Engine/Plugins/Experimental/MediaViewer/Source/MediaViewer/Private/ImageViewers/LevelEditorViewportImageViewer.cpp
--- a/Engine/Plugins/Experimental/MediaViewer/Source/MediaViewer/Private/ImageViewers/LevelEditorViewportImageViewer.cpp
+++ b/Engine/Plugins/Experimental/MediaViewer/Source/MediaViewer/Private/ImageViewers/LevelEditorViewportImageViewer.cpp
@@ -231,7 +231,7 @@ TSharedPtr<FMediaViewerLibraryItem> FLevelEditorViewportImageViewer::CreateLibra
 	return MakeShared<FItem>(ImageInfo.Id, StringValue);
 }
 
-void FLevelEditorViewportImageViewer::PaintImage(FMediaImagePaintParams& InPaintParams, const FMediaImagePaintGeometry& InPaintGeometry)
+void FLevelEditorViewportImageViewer::PaintImage(FMediaImageSlatePaintParams& InPaintParams, const FMediaImageSlatePaintGeometry& InPaintGeometry)
 {
 	if (ViewportSettings.bRealTime && !ViewportWeak.IsValid())
 	{
diff --git a/Engine/Plugins/Experimental/MediaViewer/Source/MediaViewer/Private/ImageViewers/LevelEditorViewportImageViewer.h b/Engine/Plugins/Experimental/MediaViewer/Source/MediaViewer/Private/ImageViewers/LevelEditorViewportImageViewer.h
--- a/Engine/Plugins/Experimental/MediaViewer/Source/MediaViewer/Private/ImageViewers/LevelEditorViewportImageViewer.h
+++ b/Engine/Plugins/Experimental/MediaViewer/Source/MediaViewer/Private/ImageViewers/LevelEditorViewportImageViewer.h
@@ -73,7 +73,7 @@ protected:
 	FString StringValue;
 
 	//~ Begin FMediaImageViewer
-	virtual void PaintImage(FMediaImagePaintParams& InPaintParams, const FMediaImagePaintGeometry& InPaintGeometry) override;
+	virtual void PaintImage(FMediaImageSlatePaintParams& InPaintParams, const FMediaImageSlatePaintGeometry& InPaintGeometry) override;
 	//~ End FMediaImageViewer
 };
 
diff --git a/Engine/Plugins/Experimental/MediaViewer/Source/MediaViewer/Private/ImageViewers/MaterialInterfaceImageViewer.cpp b/Engine/Plugins/Experimental/MediaViewer/Source/MediaViewer/Private/ImageViewers/MaterialInterfaceImageViewer.cpp
--- a/Engine/Plugins/Experimental/MediaViewer/Source/MediaViewer/Private/ImageViewers/MaterialInterfaceImageViewer.cpp
+++ b/Engine/Plugins/Experimental/MediaViewer/Source/MediaViewer/Private/ImageViewers/MaterialInterfaceImageViewer.cpp
@@ -324,7 +324,7 @@ TSharedPtr<FStructOnScope> FMaterialInterfaceImageViewer::GetCustomSettingsOnSco
 	);
 }
 
-void FMaterialInterfaceImageViewer::PaintImage(FMediaImagePaintParams& InPaintParams, const FMediaImagePaintGeometry& InPaintGeometry)
+void FMaterialInterfaceImageViewer::PaintImage(FMediaImageSlatePaintParams& InPaintParams, const FMediaImageSlatePaintGeometry& InPaintGeometry)
 {
 	if (MaterialSettings.bRealTime)
 	{
diff --git a/Engine/Plugins/Experimental/MediaViewer/Source/MediaViewer/Private/ImageViewers/MaterialInterfaceImageViewer.h b/Engine/Plugins/Experimental/MediaViewer/Source/MediaViewer/Private/ImageViewers/MaterialInterfaceImageViewer.h
--- a/Engine/Plugins/Experimental/MediaViewer/Source/MediaViewer/Private/ImageViewers/MaterialInterfaceImageViewer.h
+++ b/Engine/Plugins/Experimental/MediaViewer/Source/MediaViewer/Private/ImageViewers/MaterialInterfaceImageViewer.h
@@ -95,7 +95,7 @@ public:
 	virtual TSharedPtr<FMediaViewerLibraryItem> CreateLibraryItem() const override;
 	virtual TOptional<TVariant<FColor, FLinearColor>> GetPixelColor(const FIntPoint& InPixelCoords, int32 InMipLevel) const override;
 	virtual TSharedPtr<FStructOnScope> GetCustomSettingsOnScope() const override;
-	virtual void PaintImage(FMediaImagePaintParams& InPaintParams, const FMediaImagePaintGeometry& InPaintGeometry) override;
+	virtual void PaintImage(FMediaImageSlatePaintParams& InPaintParams, const FMediaImageSlatePaintGeometry& InPaintGeometry) override;
 	//~ End FMediaImageViewer
 
 	//~ Begin FNotifyHook
diff --git a/Engine/Plugins/Experimental/MediaViewer/Source/MediaViewer/Private/ImageViewers/MediaImageViewer.cpp b/Engine/Plugins/Experimental/MediaViewer/Source/MediaViewer/Private/ImageViewers/MediaImageViewer.cpp
--- a/Engine/Plugins/Experimental/MediaViewer/Source/MediaViewer/Private/ImageViewers/MediaImageViewer.cpp
+++ b/Engine/Plugins/Experimental/MediaViewer/Source/MediaViewer/Private/ImageViewers/MediaImageViewer.cpp
@@ -2,6 +2,7 @@
 
 #include "ImageViewer/MediaImageViewer.h"
 
+#include "Engine/Canvas.h"
 #include "Engine/Texture.h"
 #include "Engine/TextureRenderTarget2D.h"
 #include "GameFramework/Actor.h"
@@ -125,13 +126,13 @@ FSlateClippingZone FMediaImageViewer::CreateSlateClippingZone(const FSlateRect&
 	return ClippingZone;
 }
 
-void FMediaImageViewer::Paint(FMediaImagePaintParams& InPaintParams)
+void FMediaImageViewer::Paint(FMediaImageSlatePaintParams& InPaintParams)
 {
 	FVector2D PaintOffset = GetPaintOffset(InPaintParams.ViewerSize, InPaintParams.ViewerPosition);
 	FVector2D PaintSize = GetPaintSize();
 	FPaintGeometry PaintGeometry = InPaintParams.AllottedGeometry.ToPaintGeometry(PaintSize, FSlateLayoutTransform(PaintOffset));
 
-	const FMediaImagePaintGeometry Geometry = {
+	const FMediaImageSlatePaintGeometry Geometry = {
 		MoveTemp(PaintOffset),
 		MoveTemp(PaintSize),
 		MoveTemp(PaintGeometry)
@@ -153,7 +154,7 @@ void FMediaImageViewer::Paint(FMediaImagePaintParams& InPaintParams)
 	InPaintParams.DrawElements.PopClip();
 }
 
-void FMediaImageViewer::PaintPanel(FMediaImagePaintParams& InPaintParams, const FMediaImagePaintGeometry& InPaintGeometry)
+void FMediaImageViewer::PaintPanel(FMediaImageSlatePaintParams& InPaintParams, const FMediaImageSlatePaintGeometry& InPaintGeometry)
 {
 	const FLinearColor AlphaTint = FLinearColor(1.f, 1.f, 1.f, InPaintParams.ImageOpacity);
 
@@ -186,7 +187,7 @@ void FMediaImageViewer::PaintPanel(FMediaImagePaintParams& InPaintParams, const
 	}
 }
 
-void FMediaImageViewer::PaintImage(FMediaImagePaintParams& InPaintParams, const FMediaImagePaintGeometry& InPaintGeometry)
+void FMediaImageViewer::PaintImage(FMediaImageSlatePaintParams& InPaintParams, const FMediaImageSlatePaintGeometry& InPaintGeometry)
 {
 	if (!Brush.IsValid())
 	{
@@ -205,4 +206,124 @@ void FMediaImageViewer::PaintImage(FMediaImagePaintParams& InPaintParams, const
 	++InPaintParams.LayerId;
 }
 
+void FMediaImageViewer::Paint(const FMediaImageCanvasPaintParams& InPaintParams)
+{
+	const FVector2D CanvasSize = {
+		static_cast<double>(InPaintParams.Canvas->SizeX),
+		static_cast<double>(InPaintParams.Canvas->SizeY)
+	};
+
+	const FMediaImageCanvasPaintGeometry PaintGeometry = {
+		GetPaintOffset(CanvasSize, FVector2D::ZeroVector),
+		GetPaintSize()
+	};
+
+	PaintPanel(InPaintParams, PaintGeometry);
+	PaintImage(InPaintParams, PaintGeometry);
+}
+
+void FMediaImageViewer::PaintPanel(const FMediaImageCanvasPaintParams& InPaintParams, const FMediaImageCanvasPaintGeometry& InPaintGeometry)
+{
+	const FLinearColor AlphaTint = FLinearColor(1.f, 1.f, 1.f, InPaintParams.ImageOpacity);
+
+	if (PanelSettings.BackgroundColor.IsSet())
+	{
+		static TSoftObjectPtr<UTexture> WhiteTexturePtr = TSoftObjectPtr<UTexture>(FSoftObjectPath(TEXT("/Script/Engine.Texture2D'/Engine/EngineResources/WhiteSquareTexture.WhiteSquareTexture'")));
+		UTexture* WhiteTexture = WhiteTexturePtr.LoadSynchronous();
+
+		InPaintParams.Canvas->SetLinearDrawColor(PanelSettings.BackgroundColor.GetValue() * AlphaTint);
+		
+		InPaintParams.Canvas->DrawTile(
+			WhiteTexture,
+			InPaintGeometry.Position.X,
+			InPaintGeometry.Position.Y,
+			InPaintGeometry.Size.X,
+			InPaintGeometry.Size.Y,
+			0,
+			0,
+			WhiteTexture->GetSurfaceWidth(),
+			WhiteTexture->GetSurfaceHeight()
+		);
+
+		InPaintParams.Canvas->SetLinearDrawColor(FLinearColor::White);
+	}
+
+	if (UTexture* BackgroundTexture = Cast<UTexture>(BackgroundImageBrush.GetResourceObject()))
+	{
+		InPaintParams.Canvas->DrawTile(
+			BackgroundTexture,
+			InPaintGeometry.Position.X,
+			InPaintGeometry.Position.Y,
+			InPaintGeometry.Size.X,
+			InPaintGeometry.Size.Y,
+			0,
+			0,
+			BackgroundTexture->GetSurfaceWidth(),
+			BackgroundTexture->GetSurfaceHeight()
+		);
+	}
+}
+
+void FMediaImageViewer::PaintImage(const FMediaImageCanvasPaintParams& InPaintParams, const FMediaImageCanvasPaintGeometry& InPaintGeometry)
+{
+	if (!Brush.IsValid())
+	{
+		return;
+	}
+
+	UTexture* ForegroundTexture = Cast<UTexture>(Brush->GetResourceObject());
+
+	if (!ForegroundTexture)
+	{
+		return;
+	}
+
+	switch (InPaintParams.Orientation)
+	{
+		case EOrientation::Orient_Horizontal:
+		{
+			const float PositionX = InPaintGeometry.Position.X + (InPaintParams.UVRange.GetLowerBound().GetValue() * InPaintGeometry.Size.X);
+			const float SizeX = InPaintGeometry.Size.X * InPaintParams.UVRange.Size<float>();
+			const float U = InPaintParams.UVRange.GetLowerBound().GetValue() * ForegroundTexture->GetSurfaceWidth();
+			const float UL = ForegroundTexture->GetSurfaceWidth() * InPaintParams.UVRange.Size<float>();
+
+			InPaintParams.Canvas->DrawTile(
+				ForegroundTexture,
+				PositionX,
+				InPaintGeometry.Position.Y,
+				SizeX,
+				InPaintGeometry.Size.Y,
+				U,
+				0,
+				UL,
+				ForegroundTexture->GetSurfaceHeight()
+			);
+
+			break;
+		}
+
+		case EOrientation::Orient_Vertical:
+		{
+			const float PositionY = InPaintGeometry.Position.Y + (InPaintParams.UVRange.GetLowerBound().GetValue() * InPaintGeometry.Size.Y);
+			const float SizeY = InPaintGeometry.Size.Y * InPaintParams.UVRange.Size<float>();
+			const float V = InPaintParams.UVRange.GetLowerBound().GetValue() * ForegroundTexture->GetSurfaceHeight();
+			const float VL = ForegroundTexture->GetSurfaceHeight() * InPaintParams.UVRange.Size<float>();
+
+			InPaintParams.Canvas->DrawTile(
+				ForegroundTexture,
+				InPaintGeometry.Position.X,
+				PositionY,
+				InPaintGeometry.Size.X,
+				SizeY,
+				0,
+				V,
+				ForegroundTexture->GetSurfaceWidth(),
+				VL
+			);
+
+			break;
+		}
+	}
+}
+
 } // UE::MediaViewer
diff --git a/Engine/Plugins/Experimental/MediaViewer/Source/MediaViewer/Private/ImageViewers/MediaSourceImageViewer.cpp b/Engine/Plugins/Experimental/MediaViewer/Source/MediaViewer/Private/ImageViewers/MediaSourceImageViewer.cpp
--- a/Engine/Plugins/Experimental/MediaViewer/Source/MediaViewer/Private/ImageViewers/MediaSourceImageViewer.cpp
+++ b/Engine/Plugins/Experimental/MediaViewer/Source/MediaViewer/Private/ImageViewers/MediaSourceImageViewer.cpp
@@ -497,7 +497,7 @@ FString FMediaSourceImageViewer::GetReferencerName() const
 	return ReferencerName;
 }
 
-void FMediaSourceImageViewer::PaintImage(FMediaImagePaintParams& InPaintParams, const FMediaImagePaintGeometry& InPaintGeometry)
+void FMediaSourceImageViewer::PaintImage(FMediaImageSlatePaintParams& InPaintParams, const FMediaImageSlatePaintGeometry& InPaintGeometry)
 {
 	if (TScriptInterface<IMediaStreamPlayer> Player = MediaStream->GetPlayer())
 	{
diff --git a/Engine/Plugins/Experimental/MediaViewer/Source/MediaViewer/Private/ImageViewers/MediaSourceImageViewer.h b/Engine/Plugins/Experimental/MediaViewer/Source/MediaViewer/Private/ImageViewers/MediaSourceImageViewer.h
--- a/Engine/Plugins/Experimental/MediaViewer/Source/MediaViewer/Private/ImageViewers/MediaSourceImageViewer.h
+++ b/Engine/Plugins/Experimental/MediaViewer/Source/MediaViewer/Private/ImageViewers/MediaSourceImageViewer.h
@@ -140,7 +140,7 @@ protected:
 	void AddPlayerName(const TSharedRef<SHorizontalBox>& InStatusBar);
 
 	//~ Begin FMediaImageViewer
-	virtual void PaintImage(FMediaImagePaintParams& InPaintParams, const FMediaImagePaintGeometry& InPaintGeometry) override;
+	virtual void PaintImage(FMediaImageSlatePaintParams& InPaintParams, const FMediaImageSlatePaintGeometry& InPaintGeometry) override;
 	//~ End FMediaImageViewer
 };
 
diff --git a/Engine/Plugins/Experimental/MediaViewer/Source/MediaViewer/Private/ImageViewers/NullImageViewer.cpp b/Engine/Plugins/Experimental/MediaViewer/Source/MediaViewer/Private/ImageViewers/NullImageViewer.cpp
--- a/Engine/Plugins/Experimental/MediaViewer/Source/MediaViewer/Private/ImageViewers/NullImageViewer.cpp
+++ b/Engine/Plugins/Experimental/MediaViewer/Source/MediaViewer/Private/ImageViewers/NullImageViewer.cpp
@@ -38,7 +38,7 @@ TSharedPtr<FMediaViewerLibraryItem> FNullImageViewer::CreateLibraryItem() const
 	return nullptr;
 }
 
-void FNullImageViewer::PaintImage(FMediaImagePaintParams& InPaintParams, const FMediaImagePaintGeometry& InPaintGeometry)
+void FNullImageViewer::PaintImage(FMediaImageSlatePaintParams& InPaintParams, const FMediaImageSlatePaintGeometry& InPaintGeometry)
 {
 }
 
diff --git a/Engine/Plugins/Experimental/MediaViewer/Source/MediaViewer/Private/ImageViewers/NullImageViewer.h b/Engine/Plugins/Experimental/MediaViewer/Source/MediaViewer/Private/ImageViewers/NullImageViewer.h
--- a/Engine/Plugins/Experimental/MediaViewer/Source/MediaViewer/Private/ImageViewers/NullImageViewer.h
+++ b/Engine/Plugins/Experimental/MediaViewer/Source/MediaViewer/Private/ImageViewers/NullImageViewer.h
@@ -22,7 +22,7 @@ public:
 protected:
 	//~ Begin FMediaImageViewer
 	virtual TSharedPtr<FMediaViewerLibraryItem> CreateLibraryItem() const override;
-	virtual void PaintImage(FMediaImagePaintParams& InPaintParams, const FMediaImagePaintGeometry& InPaintGeometry) override;
+	virtual void PaintImage(FMediaImageSlatePaintParams& InPaintParams, const FMediaImageSlatePaintGeometry& InPaintGeometry) override;
 	//~ End FMediaImageViewer
 };
 
diff --git a/Engine/Plugins/Experimental/MediaViewer/Source/MediaViewer/Private/ImageViewers/Texture2DImageViewer.cpp b/Engine/Plugins/Experimental/MediaViewer/Source/MediaViewer/Private/ImageViewers/Texture2DImageViewer.cpp
--- a/Engine/Plugins/Experimental/MediaViewer/Source/MediaViewer/Private/ImageViewers/Texture2DImageViewer.cpp
+++ b/Engine/Plugins/Experimental/MediaViewer/Source/MediaViewer/Private/ImageViewers/Texture2DImageViewer.cpp
@@ -319,7 +319,7 @@ TOptional<FIntPoint> FTexture2DImageViewer::GetTextureSize() const
 	return FIntPoint(TextureSettings.Texture->GetSurfaceWidth(), TextureSettings.Texture->GetSurfaceHeight());
 }
 
-void FTexture2DImageViewer::PaintImage(FMediaImagePaintParams& InPaintParams, const FMediaImagePaintGeometry& InPaintGeometry)
+void FTexture2DImageViewer::PaintImage(FMediaImageSlatePaintParams& InPaintParams, const FMediaImageSlatePaintGeometry& InPaintGeometry)
 {
 	if (!bValidImageSize)
 	{
diff --git a/Engine/Plugins/Experimental/MediaViewer/Source/MediaViewer/Private/ImageViewers/Texture2DImageViewer.h b/Engine/Plugins/Experimental/MediaViewer/Source/MediaViewer/Private/ImageViewers/Texture2DImageViewer.h
--- a/Engine/Plugins/Experimental/MediaViewer/Source/MediaViewer/Private/ImageViewers/Texture2DImageViewer.h
+++ b/Engine/Plugins/Experimental/MediaViewer/Source/MediaViewer/Private/ImageViewers/Texture2DImageViewer.h
@@ -104,7 +104,7 @@ protected:
 	TOptional<FIntPoint> GetTextureSize() const;
 
 	//~ Begin FMediaImageViewer
-	virtual void PaintImage(FMediaImagePaintParams& InPaintParams, const FMediaImagePaintGeometry& InPaintGeometry) override;
+	virtual void PaintImage(FMediaImageSlatePaintParams& InPaintParams, const FMediaImageSlatePaintGeometry& InPaintGeometry) override;
 	//~ End FMediaImageViewer
 };
 
diff --git a/Engine/Plugins/Experimental/MediaViewer/Source/MediaViewer/Private/ImageViewers/ViewportImageViewer.cpp b/Engine/Plugins/Experimental/MediaViewer/Source/MediaViewer/Private/ImageViewers/ViewportImageViewer.cpp
--- a/Engine/Plugins/Experimental/MediaViewer/Source/MediaViewer/Private/ImageViewers/ViewportImageViewer.cpp
+++ b/Engine/Plugins/Experimental/MediaViewer/Source/MediaViewer/Private/ImageViewers/ViewportImageViewer.cpp
@@ -75,7 +75,7 @@ TSharedPtr<FStructOnScope> FViewportImageViewer::GetCustomSettingsOnScope() cons
 	);
 }
 
-void FViewportImageViewer::PaintImage(FMediaImagePaintParams& InPaintParams, const FMediaImagePaintGeometry& InPaintGeometry)
+void FViewportImageViewer::PaintImage(FMediaImageSlatePaintParams& InPaintParams, const FMediaImageSlatePaintGeometry& InPaintGeometry)
 {
 	if (RenderState == ERenderState::NotRendered || ViewportSettings.bRealTime)
 	{
diff --git a/Engine/Plugins/Experimental/MediaViewer/Source/MediaViewer/Private/ImageViewers/ViewportImageViewer.h b/Engine/Plugins/Experimental/MediaViewer/Source/MediaViewer/Private/ImageViewers/ViewportImageViewer.h
--- a/Engine/Plugins/Experimental/MediaViewer/Source/MediaViewer/Private/ImageViewers/ViewportImageViewer.h
+++ b/Engine/Plugins/Experimental/MediaViewer/Source/MediaViewer/Private/ImageViewers/ViewportImageViewer.h
@@ -79,7 +79,7 @@ protected:
 	void OnRenderComplete(bool bInSuccess);
 
 	//~ Begin FMediaImageViewer
-	virtual void PaintImage(FMediaImagePaintParams& InPaintParams, const FMediaImagePaintGeometry& InPaintGeometry) override;
+	virtual void PaintImage(FMediaImageSlatePaintParams& InPaintParams, const FMediaImageSlatePaintGeometry& InPaintGeometry) override;
 	//~ End FMediaImageViewer
 };
 
diff --git a/Engine/Plugins/Experimental/MediaViewer/Source/MediaViewer/Private/MediaViewerCommands.cpp b/Engine/Plugins/Experimental/MediaViewer/Source/MediaViewer/Private/MediaViewerCommands.cpp
--- a/Engine/Plugins/Experimental/MediaViewer/Source/MediaViewer/Private/MediaViewerCommands.cpp
+++ b/Engine/Plugins/Experimental/MediaViewer/Source/MediaViewer/Private/MediaViewerCommands.cpp
@@ -280,6 +280,15 @@ void FMediaViewerCommands::RegisterCommands()
 		EUserInterfaceActionType::Button,
 		FInputChord(EKeys::Slash)
 	);
+
+	UI_COMMAND(
+		Snapshot,
+		"Snapshot",
+		"Take a snapshot of the current viewer display.",
+		EUserInterfaceActionType::Button,
+		FInputChord(EKeys::S, EModifierKey::Control),
+		FInputChord(EKeys::S, EModifierKey::Command)
+	);
 }
 
 } // UE::MediaViewer::Private
diff --git a/Engine/Plugins/Experimental/MediaViewer/Source/MediaViewer/Private/MediaViewerCommands.h b/Engine/Plugins/Experimental/MediaViewer/Source/MediaViewer/Private/MediaViewerCommands.h
--- a/Engine/Plugins/Experimental/MediaViewer/Source/MediaViewer/Private/MediaViewerCommands.h
+++ b/Engine/Plugins/Experimental/MediaViewer/Source/MediaViewer/Private/MediaViewerCommands.h
@@ -58,6 +58,8 @@ public:
 	TSharedPtr<FUICommandInfo> AddToLibrary;
 
 	TSharedPtr<FUICommandInfo> SwapAB;
+
+	TSharedPtr<FUICommandInfo> Snapshot;
 };
 
 } // UE::MediaViewer::Private
diff --git a/Engine/Plugins/Experimental/MediaViewer/Source/MediaViewer/Private/Widgets/SMediaViewer.cpp b/Engine/Plugins/Experimental/MediaViewer/Source/MediaViewer/Private/Widgets/SMediaViewer.cpp
--- a/Engine/Plugins/Experimental/MediaViewer/Source/MediaViewer/Private/Widgets/SMediaViewer.cpp
+++ b/Engine/Plugins/Experimental/MediaViewer/Source/MediaViewer/Private/Widgets/SMediaViewer.cpp
@@ -2,31 +2,45 @@
 
 #include "Widgets/SMediaViewer.h"
 
+#include "AssetRegistry/AssetRegistryModule.h"
 #include "Brushes/SlateColorBrush.h"
 #include "Brushes/SlateImageBrush.h"
+#include "CanvasTypes.h"
+#include "ContentBrowserModule.h"
+#include "Engine/Canvas.h"
+#include "Engine/Engine.h"
 #include "Engine/Texture2D.h"
+#include "Engine/TextureRenderTarget2D.h"
 #include "Framework/Application/SlateApplication.h"
+#include "Framework/Notifications/NotificationManager.h"
+#include "IContentBrowserSingleton.h"
 #include "ImageViewer/MediaImageViewer.h"
 #include "ImageViewers/NullImageViewer.h"
 #include "IMediaViewerModule.h"
 #include "Library/MediaViewerLibrary.h"
-#include "Library/MediaViewerLibraryIni.h"
 #include "Library/MediaViewerLibraryGroup.h"
+#include "Library/MediaViewerLibraryIni.h"
 #include "MediaViewerCommands.h"
+#include "MediaViewerUtils.h"
 #include "Misc/MessageDialog.h"
 #include "Sidebar/SSidebar.h"
 #include "Sidebar/SSidebarContainer.h"
+#include "TextureResource.h"
 #include "Widgets/Colors/SColorBlock.h"
 #include "Widgets/Layout/SBox.h"
 #include "Widgets/Layout/SSplitter.h"
 #include "Widgets/MediaViewerDelegates.h"
+#include "Widgets/Notifications/SNotificationList.h"
 #include "Widgets/SMediaImageViewerOverlay.h"
 #include "Widgets/SMediaViewerLibrary.h"
 #include "Widgets/SMediaViewerLibraryPrivate.h"
 #include "Widgets/SMediaViewerToolbar.h"
+#include "UObject/StrongObjectPtr.h"
 
 #define LOCTEXT_NAMESPACE "SMediaViewer"
 
+extern ENGINE_API UEngine* GEngine;
+
 namespace UE::MediaViewer::Private
 {
 
@@ -217,6 +231,12 @@ void SMediaViewer::BindCommands()
 	);
 
 	CommandList->MapAction(Commands.ResetAllTransforms, FExecuteAction::CreateSP(this, &SMediaViewer::ResetTransformToAll));
+
+	CommandList->MapAction(
+		Commands.Snapshot,
+		FExecuteAction::CreateSP(this, &SMediaViewer::TakeSnapshot),
+		FCanExecuteAction::CreateSP(this, &SMediaViewer::CanTakeSnapshot)
+	);
 }
 
 bool SMediaViewer::IsShowingBothImageViewers() const
@@ -423,7 +443,7 @@ int32 SMediaViewer::OnPaint(const FPaintArgs& InArgs, const FGeometry& InAllotte
 
 	FFloatRange UVRange = FFloatRange(0.f, 1.f);
 
-	FMediaImagePaintParams PaintParams = {
+	FMediaImageSlatePaintParams PaintParams = {
 		InArgs,
 		InAllottedGeometry,
 		InMyCullingRect,
@@ -1011,6 +1031,203 @@ void SMediaViewer::CheckLoadState()
 	LoadState(0);
 }
 
+bool SMediaViewer::CanTakeSnapshot() const
+{
+	switch (ActiveView)
+	{
+		case EMediaImageViewerActivePosition::Single:
+			return !IsImageViewerNull(EMediaImageViewerPosition::First);
+
+		case EMediaImageViewerActivePosition::Both:
+			return !IsImageViewerNull(EMediaImageViewerPosition::First)
+				|| !IsImageViewerNull(EMediaImageViewerPosition::Second);
+
+		default:
+			return false;
+	}
+}
+
+void SMediaViewer::TakeSnapshot()
+{
+	if (!ContentSlot)
+	{
+		return;
+	}
+
+	if (IsImageViewerNull(EMediaImageViewerPosition::First)
+		&& IsImageViewerNull(EMediaImageViewerPosition::Second))
+	{
+		return;
+	}
+
+	if (FMath::IsNearlyZero(ViewerSize.X) || FMath::IsNearlyZero(ViewerSize.Y))
+	{
+		return;
+	}
+
+	TStrongObjectPtr<UTextureRenderTarget2D> RenderTarget = TStrongObjectPtr<UTextureRenderTarget2D>(FMediaViewerUtils::CreateRenderTarget(
+		FIntPoint(
+			FMath::RoundToInt(ViewerSize.X),
+			FMath::RoundToInt(ViewerSize.Y)
+		),
+		/* Transparent */ true
+	));
+
+	if (!RenderTarget)
+	{
+		return;
+	}
+
+	// Choose asset location
+	IContentBrowserSingleton& ContentBrowser = FModuleManager::LoadModuleChecked<FContentBrowserModule>("ContentBrowser").Get();
+	const FContentBrowserItemPath CurrentPath = ContentBrowser.GetCurrentPath();
+	const FString PathStr = CurrentPath.HasInternalPath() ? CurrentPath.GetInternalPathString() : "/Game";
+
+	FSaveAssetDialogConfig SaveAssetDialogConfig;
+	SaveAssetDialogConfig.DialogTitleOverride = LOCTEXT("ExportMaterialTo", "Export Material To");
+	SaveAssetDialogConfig.DefaultPath = PathStr;
+	SaveAssetDialogConfig.ExistingAssetPolicy = ESaveAssetDialogExistingAssetPolicy::AllowButWarn;
+	SaveAssetDialogConfig.DefaultAssetName = LOCTEXT("MediaViewerSnapshot", "MediaViewerSnapshot").ToString();
+
+	FContentBrowserModule& ContentBrowserModule = FModuleManager::LoadModuleChecked<FContentBrowserModule>("ContentBrowser");
+	FString SaveObjectPath = ContentBrowserModule.Get().CreateModalSaveAssetDialog(SaveAssetDialogConfig);
+
+	if (SaveObjectPath.IsEmpty())
+	{
+		return;
+	}
+
+	FTextureRenderTargetResource* RenderTargetResource = RenderTarget->GameThread_GetRenderTargetResource();
+
+	UCanvas* Canvas = NewObject<UCanvas>(GetTransientPackage());
+
+	FCanvas RenderCanvas(
+		RenderTargetResource,
+		nullptr,
+		FGameTime::CreateUndilated(0, 0),
+		GEngine->GetDefaultWorldFeatureLevel()
+	);
+
+	Canvas->Init(RenderTarget->SizeX, RenderTarget->SizeY, nullptr, &RenderCanvas);
+
+	{
+		RHI_BREADCRUMB_EVENT_GAMETHREAD_F("CreateMediaViewerSnapshot", "CreateMediaViewerSnapshot: %s", RenderTarget->GetFName());
+
+		ENQUEUE_RENDER_COMMAND(FlushDeferredResourceUpdateCommand)(
+			[RenderTargetResource](FRHICommandListImmediate& RHICmdList)
+			{
+				RenderTargetResource->FlushDeferredResourceUpdate(RHICmdList);
+			});
+
+		DrawImageViewers(Canvas);
+
+		RenderCanvas.Flush_GameThread();
+		Canvas->Canvas = nullptr;
+
+		//UpdateResourceImmediate must be called here to ensure mips are generated.
+		RenderTarget->UpdateResourceImmediate(false);
+
+		ENQUEUE_RENDER_COMMAND(ResetSceneTextureExtentHistory)(
+			[RenderTargetResource](FRHICommandListImmediate& RHICmdList)
+			{
+				RenderTargetResource->ResetSceneTextureExtentsHistory();
+			});
+	}
+
+	const FString PackageName = FPaths::GetBaseFilename(*SaveObjectPath, false);
+	const FString AssetName = FPaths::GetBaseFilename(*SaveObjectPath, true);
+
+	UTexture2D* Texture = RenderTarget->ConstructTexture2D(GetTransientPackage(), AssetName, RF_Public | RF_Standalone);
+
+	if (!Texture)
+	{
+		return;
+	}
+
+	RenderTarget->ReleaseResource();
+
+	// Create package and save texture.
+	UPackage* Package = CreatePackage(*PackageName);
+
+	if (!Package)
+	{
+		return;
+	}
+
+	Texture->Rename(*AssetName, Package, REN_DontCreateRedirectors | REN_DoNotDirty | REN_ForceNoResetLoaders | REN_NonTransactional);
+
+	FAssetRegistryModule::AssetCreated(Package);
+
+	// Make a toast!
+	const FString HyperLinkText = FPaths::ConvertRelativePathToFull(SaveObjectPath);
+
+	FNotificationInfo Info(LOCTEXT("SnapshotCreated", "Snapshot created."));
+
+	Info.Hyperlink = FSimpleDelegate::CreateStatic([](FString SourceFilePath)
+		{
+			FPlatformProcess::ExploreFolder(*(FPaths::GetPath(SourceFilePath)));
+		}, HyperLinkText);
+
+	Info.HyperlinkText = FText::FromString(HyperLinkText);
+	Info.ExpireDuration = 3.0f;
+
+	FSlateNotificationManager::Get().AddNotification(Info);
+}
+
+void SMediaViewer::DrawImageViewers(UCanvas* InCanvas) const
+{
+	if (UTexture* BackgroundTexture = MediaViewerSettings.Texture.LoadSynchronous())
+	{
+		const float BackgroundWidth = BackgroundTexture->GetSurfaceWidth();
+		const float BackgroundHeight = BackgroundTexture->GetSurfaceHeight();
+
+		InCanvas->DrawTile(
+			BackgroundTexture,
+			/* X */ 0,
+			/* Y */ 0,
+			/* XL */ BackgroundWidth,
+			/* YL */ BackgroundHeight,
+			/* U */ 0,
+			/* V */ 0,
+			/* UL */ static_cast<float>(InCanvas->SizeX) / BackgroundWidth,
+			/* VL */ static_cast<float>(InCanvas->SizeY) / BackgroundHeight
+		);
+	}
+
+	FFloatRange UVRange = FFloatRange(0.f, 1.f);
+
+	FMediaImageCanvasPaintParams PaintParams = {
+		InCanvas,
+		UVRange,
+		MediaViewerSettings.ABOrientation,
+		1.f
+	};
+
+	const float ABSplitterLocation = MediaViewerSettings.ABSplitterLocation / 100.f;
+	const float SecondImageOpacity = MediaViewerSettings.SecondImageOpacity / 100.f;
+
+	switch (ActiveView)
+	{
+		case EMediaImageViewerActivePosition::Single:
+			GetImageViewer(EMediaImageViewerPosition::First)->Paint(PaintParams);
+			break;
+
+		case EMediaImageViewerActivePosition::Both:
+			UVRange.SetLowerBound(0.f);
+			UVRange.SetUpperBound(1.f);
+			GetImageViewer(EMediaImageViewerPosition::First)->Paint(PaintParams);
+
+			if (!FMath::IsNearlyZero(ABSplitterLocation - 1.f))
+			{
+				UVRange.SetLowerBound(ABSplitterLocation);
+				UVRange.SetUpperBound(1.f);
+				PaintParams.ImageOpacity = SecondImageOpacity;
+				GetImageViewer(EMediaImageViewerPosition::Second)->Paint(PaintParams);
+			}
+			break;
+	}
+}
+
 } // UE::MediaViewer::Private
 
 #undef LOCTEXT_NAMESPACE
diff --git a/Engine/Plugins/Experimental/MediaViewer/Source/MediaViewer/Private/Widgets/SMediaViewer.h b/Engine/Plugins/Experimental/MediaViewer/Source/MediaViewer/Private/Widgets/SMediaViewer.h
--- a/Engine/Plugins/Experimental/MediaViewer/Source/MediaViewer/Private/Widgets/SMediaViewer.h
+++ b/Engine/Plugins/Experimental/MediaViewer/Source/MediaViewer/Private/Widgets/SMediaViewer.h
@@ -16,6 +16,7 @@
 
 class FUICommandList;
 class SViewportToolBar;
+class UCanvas;
 
 namespace UE::MediaViewer
 {
@@ -237,6 +238,11 @@ protected:
 	void ToggleLockedTransform();
 
 	void CheckLoadState();
+
+	bool CanTakeSnapshot() const;
+	void TakeSnapshot();
+
+	void DrawImageViewers(UCanvas* InCanvas) const;
 };
 
 } // UE::MediaViewer::Private
diff --git a/Engine/Plugins/Experimental/MediaViewer/Source/MediaViewer/Private/Widgets/SMediaViewerToolbar.cpp b/Engine/Plugins/Experimental/MediaViewer/Source/MediaViewer/Private/Widgets/SMediaViewerToolbar.cpp
--- a/Engine/Plugins/Experimental/MediaViewer/Source/MediaViewer/Private/Widgets/SMediaViewerToolbar.cpp
+++ b/Engine/Plugins/Experimental/MediaViewer/Source/MediaViewer/Private/Widgets/SMediaViewerToolbar.cpp
@@ -255,9 +255,17 @@ TSharedRef<SWidget> SMediaViewerToolbar::MakeCenterToolbar()
 			);
 		}
 		ToolbarBuilder.EndBlockGroup();
-
+		
 		ToolbarBuilder.BeginBlockGroup();
 		{
+			ToolbarBuilder.AddToolBarButton(
+				Commands.Snapshot,
+				NAME_None,
+				FText::GetEmpty(),
+				TAttribute<FText>(),
+				FSlateIcon(FAppStyle::Get().GetStyleSetName(), "LevelViewport.HighResScreenshot")
+			);
+
 			ToolbarBuilder.AddComboButton(
 				FUIAction(),
 				FOnGetContent::CreateSP(this, &SMediaViewerToolbar::GetBackgroundTextureSettingsWidget),
diff --git a/Engine/Plugins/Experimental/MediaViewer/Source/MediaViewer/Public/ImageViewer/MediaImageViewer.h b/Engine/Plugins/Experimental/MediaViewer/Source/MediaViewer/Public/ImageViewer/MediaImageViewer.h
--- a/Engine/Plugins/Experimental/MediaViewer/Source/MediaViewer/Public/ImageViewer/MediaImageViewer.h
+++ b/Engine/Plugins/Experimental/MediaViewer/Source/MediaViewer/Public/ImageViewer/MediaImageViewer.h
@@ -30,6 +30,7 @@ class FStructOnScope;
 class FUICommandList;
 class FWidgetStyle;
 class SWidget;
+class UCanvas;
 class UTexture;
 struct FFloatRange;
 struct FGeometry;
@@ -125,7 +126,7 @@ struct FMediaImageViewerInfo
 /**
  * Parameters used by the image viewers to paint their image.
  */
-struct FMediaImagePaintParams
+struct FMediaImageSlatePaintParams
 {
 	/** Slate paint args. */
 	const FPaintArgs& Args;
@@ -156,7 +157,7 @@ struct FMediaImagePaintParams
 };
 
 /** Geometry of the image within the panel. */
-struct FMediaImagePaintGeometry
+struct FMediaImageSlatePaintGeometry
 {
 	/** Position of the image within the painted area. */
 	const FVector2D Position;
@@ -166,6 +167,30 @@ struct FMediaImagePaintGeometry
 	const FPaintGeometry Geometry;
 };
 
+/**
+ * Parameters used by the image viewers to paint their image.
+ */
+struct FMediaImageCanvasPaintParams
+{
+	/* The canvas to draw on. */
+	UCanvas* Canvas;
+	/** UV Range to draw the image on, either horizontally (L->R) and vertically (T->B). */
+	const FFloatRange& UVRange;
+	/** Orientation of the panel split. */
+	const EOrientation Orientation;
+	/** Opacity to draw the image with. 0-1. */
+	float ImageOpacity;
+};
+
+/** Geometry of the image within the panel. */
+struct FMediaImageCanvasPaintGeometry
+{
+	/** Position of the image within the painted area. */
+	const FVector2D Position;
+	/** Size of the image within the painted area. */
+	const FVector2D Size;
+};
+
 /**
  * Image resource for the AB Image Viewer representing a source that can be displayed as a 2d image.
  */
@@ -208,7 +233,10 @@ public:
 	virtual TOptional<TVariant<FColor, FLinearColor>> GetPixelColor(const FIntPoint& InPixelCoords, int32 InMipLevel) const = 0;
 
 	/** Paint the panel and image on to the viewer. */
-	void Paint(FMediaImagePaintParams& InPaintParams);
+	void Paint(FMediaImageSlatePaintParams& InPaintParams);
+
+	/** Paint the panel and image to a canvas. */
+	void Paint(const FMediaImageCanvasPaintParams& InPaintParams);
 
 	/**
 	 * Extend the tool bar (top). @see MediaViewer.h
@@ -298,9 +326,13 @@ protected:
 	FSlateClippingZone CreateSlateClippingZone(const FSlateRect& InCullingRect, float InDPIScale, const FVector2D& InViewerPosition, 
 		EOrientation InOrientation, const FFloatRange& InUVRange) const;
 
-	MEDIAVIEWER_API virtual void PaintPanel(FMediaImagePaintParams& InPaintParams, const FMediaImagePaintGeometry& InPaintGeometry);
+	MEDIAVIEWER_API virtual void PaintPanel(FMediaImageSlatePaintParams& InPaintParams, const FMediaImageSlatePaintGeometry& InPaintGeometry);
+
+	MEDIAVIEWER_API virtual void PaintImage(FMediaImageSlatePaintParams& InPaintParams, const FMediaImageSlatePaintGeometry& InPaintGeometry);
+
+	MEDIAVIEWER_API virtual void PaintPanel(const FMediaImageCanvasPaintParams& InPaintParams, const FMediaImageCanvasPaintGeometry& InPaintGeometry);
 
-	MEDIAVIEWER_API virtual void PaintImage(FMediaImagePaintParams& InPaintParams, const FMediaImagePaintGeometry& InPaintGeometry);
+	MEDIAVIEWER_API virtual void PaintImage(const FMediaImageCanvasPaintParams& InPaintParams, const FMediaImageCanvasPaintGeometry& InPaintGeometry);
 };
 
 } // UE::MediaViewer
