Horde issue 934988

diff --git a/Engine/Plugins/Experimental/MediaViewer/Source/MediaViewer/Private/MediaViewerCommands.cpp b/Engine/Plugins/Experimental/MediaViewer/Source/MediaViewer/Private/MediaViewerCommands.cpp
--- a/Engine/Plugins/Experimental/MediaViewer/Source/MediaViewer/Private/MediaViewerCommands.cpp
+++ b/Engine/Plugins/Experimental/MediaViewer/Source/MediaViewer/Private/MediaViewerCommands.cpp
@@ -286,8 +286,7 @@ void FMediaViewerCommands::RegisterCommands()
 		"Snapshot",
 		"Take a snapshot of the current viewer display.",
 		EUserInterfaceActionType::Button,
-		FInputChord(EKeys::S, EModifierKey::Control),
-		FInputChord(EKeys::S, EModifierKey::Command)
+		FInputChord(EKeys::S, EModifierKey::Alt)
 	);
 }
 
