Media Viewer: Snapshot category

diff --git a/Engine/Plugins/Experimental/MediaViewer/Source/MediaViewer/Private/Library/MediaViewerLibrary.cpp b/Engine/Plugins/Experimental/MediaViewer/Source/MediaViewer/Private/Library/MediaViewerLibrary.cpp
--- a/Engine/Plugins/Experimental/MediaViewer/Source/MediaViewer/Private/Library/MediaViewerLibrary.cpp
+++ b/Engine/Plugins/Experimental/MediaViewer/Source/MediaViewer/Private/Library/MediaViewerLibrary.cpp
@@ -12,8 +12,9 @@
 
 namespace UE::MediaViewer::Private
 {
-	const FGuid DefaultGroupId = FGuid(0xD42C17B4, 0x47939576, 0x3BDB9594, 0x5101041A);
-	const FGuid HistoryGroupId = FGuid(0x35C7E8CA, 0x2B6D1A76, 0xA9EAFE80, 0xA2B12B41);
+	constexpr FGuid DefaultGroupId = FGuid(0xD42C17B4, 0x47939576, 0x3BDB9594, 0x5101041A);
+	constexpr FGuid HistoryGroupId = FGuid(0x35C7E8CA, 0x2B6D1A76, 0xA9EAFE80, 0xA2B12B41);
+	constexpr FGuid SnapshotGroupId = FGuid(0X11F83073, 0X7FF4BC41, 0XA4FFE792, 0XD073F41F);
 }
 
 namespace UE::MediaViewer::Private
@@ -29,6 +30,11 @@ const FGuid& FMediaViewerLibrary::GetHistoryGroupId() const
 	return HistoryGroupId;
 }
 
+const FGuid& FMediaViewerLibrary::GetSnapshotsGroupId() const
+{
+	return SnapshotGroupId;
+}
+
 FMediaViewerLibrary::FMediaViewerLibrary(const FArgs& InArgs)
 {
 	const FText DefaultTitle = InArgs.DefaultGroupName.IsEmpty()
@@ -46,6 +52,13 @@ FMediaViewerLibrary::FMediaViewerLibrary(const FArgs& InArgs)
 		/* Dynamic */ false
 	));
 
+	Groups.Add(MakeShared<FMediaViewerLibraryGroup>(
+		GetSnapshotsGroupId(),
+		LOCTEXT("Snapshots", "Snapshots"),
+		LOCTEXT("SnapshotsTooltip", "Snapshots taken with the Media Viewer tool."),
+		/* Dynamic */ false
+	));
+
 	Groups.Add(MakeShared<FMediaViewerLibraryGroup>(
 		GetHistoryGroupId(),
 		LOCTEXT("History", "History"),
@@ -114,7 +127,7 @@ TSharedPtr<FMediaViewerLibraryGroup> FMediaViewerLibrary::GetGroup(const FGuid&
 
 bool FMediaViewerLibrary::CanRemoveGroup(const FGuid& InGroupIdToRemove) const
 {
-	if (InGroupIdToRemove == GetDefaultGroupId() || InGroupIdToRemove == GetHistoryGroupId())
+	if (InGroupIdToRemove == GetDefaultGroupId() || InGroupIdToRemove == GetHistoryGroupId() || InGroupIdToRemove == GetSnapshotsGroupId())
 	{
 		return false;
 	}
@@ -190,7 +203,7 @@ TSharedPtr<FMediaViewerLibraryGroup> FMediaViewerLibrary::GetItemGroup(const FGu
 {
 	for (const TSharedRef<FMediaViewerLibraryGroup>& Group : Groups)
 	{
-		if (Group->GetId() == GetHistoryGroupId())
+		if (Group->GetId() == GetHistoryGroupId() || Group->GetId() == GetSnapshotsGroupId())
 		{
 			continue;
 		}
diff --git a/Engine/Plugins/Experimental/MediaViewer/Source/MediaViewer/Private/Library/MediaViewerLibrary.h b/Engine/Plugins/Experimental/MediaViewer/Source/MediaViewer/Private/Library/MediaViewerLibrary.h
--- a/Engine/Plugins/Experimental/MediaViewer/Source/MediaViewer/Private/Library/MediaViewerLibrary.h
+++ b/Engine/Plugins/Experimental/MediaViewer/Source/MediaViewer/Private/Library/MediaViewerLibrary.h
@@ -36,6 +36,7 @@ public:
 	//~ Begin IMediaViewerLibrary
 	virtual const FGuid& GetDefaultGroupId() const override;
 	virtual const FGuid& GetHistoryGroupId() const override;
+	virtual const FGuid& GetSnapshotsGroupId() const override;
 	virtual const TArray<TSharedRef<const FMediaViewerLibraryGroup>>& GetGroups() const override;
 	virtual bool AddGroup(const TSharedRef<FMediaViewerLibraryGroup>& InNewGroup) override;
 	virtual TSharedPtr<FMediaViewerLibraryGroup> GetGroup(const FGuid& InGroupIdToRemove) const override;
diff --git a/Engine/Plugins/Experimental/MediaViewer/Source/MediaViewer/Private/Widgets/SMediaViewer.cpp b/Engine/Plugins/Experimental/MediaViewer/Source/MediaViewer/Private/Widgets/SMediaViewer.cpp
--- a/Engine/Plugins/Experimental/MediaViewer/Source/MediaViewer/Private/Widgets/SMediaViewer.cpp
+++ b/Engine/Plugins/Experimental/MediaViewer/Source/MediaViewer/Private/Widgets/SMediaViewer.cpp
@@ -16,6 +16,7 @@
 #include "IContentBrowserSingleton.h"
 #include "ImageViewer/MediaImageViewer.h"
 #include "ImageViewers/NullImageViewer.h"
+#include "ImageViewers/Texture2DImageViewer.h"
 #include "IMediaViewerModule.h"
 #include "Library/MediaViewerLibrary.h"
 #include "Library/MediaViewerLibraryGroup.h"
@@ -1158,6 +1159,17 @@ void SMediaViewer::TakeSnapshot()
 
 	FAssetRegistryModule::AssetCreated(Package);
 
+	// Add to snapshots group
+	TSharedRef<FMediaViewerLibraryItem> SnapshotItem = MakeShared<FTexture2DImageViewer::FItem>(
+		FText::FromString(AssetName),
+		LOCTEXT("SnapshotTooltip", "Snapshot from Media Viewer"),
+		/* Transient */ false,
+		TNotNull<UTexture2D*>(Texture)
+	);
+
+	TSharedRef<FMediaViewerLibrary> LibraryActual = GetLibrary();
+	LibraryActual->AddItemToGroup(SnapshotItem, LibraryActual->GetSnapshotsGroupId());
+
 	// Make a toast!
 	const FString HyperLinkText = FPaths::ConvertRelativePathToFull(SaveObjectPath);
 
diff --git a/Engine/Plugins/Experimental/MediaViewer/Source/MediaViewer/Public/Library/IMediaViewerLibrary.h b/Engine/Plugins/Experimental/MediaViewer/Source/MediaViewer/Public/Library/IMediaViewerLibrary.h
--- a/Engine/Plugins/Experimental/MediaViewer/Source/MediaViewer/Public/Library/IMediaViewerLibrary.h
+++ b/Engine/Plugins/Experimental/MediaViewer/Source/MediaViewer/Public/Library/IMediaViewerLibrary.h
@@ -73,6 +73,11 @@ public:
 	 */
 	virtual const FGuid& GetHistoryGroupId() const = 0;
 
+	/**
+	 * @returns The id of the snapshot group.
+	 */
+	virtual const FGuid& GetSnapshotsGroupId() const = 0;
+
 	/**
 	 * Return the list of groups.
 	 */
