Media Viewer: Added texture size constraits to render calls.

diff --git a/Engine/Plugins/Experimental/MediaViewer/Source/MediaViewer/Private/ImageViewers/LevelEditorViewportImageViewer.cpp b/Engine/Plugins/Experimental/MediaViewer/Source/MediaViewer/Private/ImageViewers/LevelEditorViewportImageViewer.cpp
--- a/Engine/Plugins/Experimental/MediaViewer/Source/MediaViewer/Private/ImageViewers/LevelEditorViewportImageViewer.cpp
+++ b/Engine/Plugins/Experimental/MediaViewer/Source/MediaViewer/Private/ImageViewers/LevelEditorViewportImageViewer.cpp
@@ -163,7 +163,7 @@ UTextureRenderTarget2D* FLevelEditorViewportImageViewer::FItem::CreateRenderTarg
 		return nullptr;
 	}
 
-	const FIntPoint RenderTargetSize = {32, 32};
+	const FIntPoint RenderTargetSize = {128, 128};
 	UTextureRenderTarget2D* RenderTarget = FMediaViewerUtils::CreateRenderTarget(RenderTargetSize, /* Transparent */ false);
 
 	RenderViewport(
diff --git a/Engine/Plugins/Experimental/MediaViewer/Source/MediaViewer/Private/ImageViewers/TextureSampleCache.cpp b/Engine/Plugins/Experimental/MediaViewer/Source/MediaViewer/Private/ImageViewers/TextureSampleCache.cpp
--- a/Engine/Plugins/Experimental/MediaViewer/Source/MediaViewer/Private/ImageViewers/TextureSampleCache.cpp
+++ b/Engine/Plugins/Experimental/MediaViewer/Source/MediaViewer/Private/ImageViewers/TextureSampleCache.cpp
@@ -227,6 +227,14 @@ void FTextureSampleCache::SetPixelColor_RenderTarget(const FIntPoint& InPixelCoo
 				return;
 			}
 
+			constexpr uint32 MinDimension = 3;
+
+			if (SourceResource->GetSizeX() < MinDimension || SourceResource->GetSizeY() < MinDimension)
+			{
+				OnComplete.ExecuteIfBound(false);
+				return;
+			}
+
 			FTextureResource* TargetResource = Target->GetResource();
 
 			if (!TargetResource)
@@ -235,6 +243,12 @@ void FTextureSampleCache::SetPixelColor_RenderTarget(const FIntPoint& InPixelCoo
 				return;
 			}
 
+			if (TargetResource->GetSizeX() < MinDimension || TargetResource->GetSizeY() < MinDimension)
+			{
+				OnComplete.ExecuteIfBound(false);
+				return;
+			}
+
 			FRDGBuilder GraphBuilder(InRHICmdList);
 			
 			FRDGTexture* InputTexture = GraphBuilder.RegisterExternalTexture(CreateRenderTarget(SourceResource->TextureRHI, TEXT("SourceTexture")));
diff --git a/Engine/Plugins/Experimental/MediaViewer/Source/MediaViewer/Private/ImageViewers/ViewportImageViewer.cpp b/Engine/Plugins/Experimental/MediaViewer/Source/MediaViewer/Private/ImageViewers/ViewportImageViewer.cpp
--- a/Engine/Plugins/Experimental/MediaViewer/Source/MediaViewer/Private/ImageViewers/ViewportImageViewer.cpp
+++ b/Engine/Plugins/Experimental/MediaViewer/Source/MediaViewer/Private/ImageViewers/ViewportImageViewer.cpp
@@ -153,6 +153,14 @@ void FViewportImageViewer::RenderViewport(TNotNull<FViewport*> InViewport, TNotN
 				return;
 			}
 
+			constexpr int32 MinDimension = 3;
+
+			if (InputTexture->Desc.Extent.X < MinDimension || InputTexture->Desc.Extent.Y < MinDimension)
+			{
+				InRenderComplete.ExecuteIfBound(false);
+				return;
+			}
+
 			FRDGTexture* OutputTexture = RenderTargetResource->GetRenderTargetTexture(GraphBuilder);
 
 			if (!OutputTexture)
@@ -161,6 +169,12 @@ void FViewportImageViewer::RenderViewport(TNotNull<FViewport*> InViewport, TNotN
 				return;
 			}
 
+			if (OutputTexture->Desc.Extent.X < MinDimension || OutputTexture->Desc.Extent.Y < MinDimension)
+			{
+				InRenderComplete.ExecuteIfBound(false);
+				return;
+			}
+
 			const FGlobalShaderMap* GlobalShaderMap = GetGlobalShaderMap(GMaxRHIFeatureLevel);
 
 			AddDrawTexturePass(GraphBuilder, GlobalShaderMap, InputTexture, OutputTexture, FRDGDrawTextureInfo());
