Media Viewer: "Saved" category renamed to "Pinned".
- Fixed an incosistency with the Pin command.

diff --git a/Engine/Plugins/Experimental/MediaViewer/Source/MediaViewer/Private/Library/MediaViewerLibrary.cpp b/Engine/Plugins/Experimental/MediaViewer/Source/MediaViewer/Private/Library/MediaViewerLibrary.cpp
--- a/Engine/Plugins/Experimental/MediaViewer/Source/MediaViewer/Private/Library/MediaViewerLibrary.cpp
+++ b/Engine/Plugins/Experimental/MediaViewer/Source/MediaViewer/Private/Library/MediaViewerLibrary.cpp
@@ -38,11 +38,11 @@ const FGuid& FMediaViewerLibrary::GetSnapshotsGroupId() const
 FMediaViewerLibrary::FMediaViewerLibrary(const FArgs& InArgs)
 {
 	const FText DefaultTitle = InArgs.DefaultGroupName.IsEmpty()
-		? LOCTEXT("DefaultGroup", "Saved")
+		? LOCTEXT("DefaultGroup", "Pinned")
 		: InArgs.DefaultGroupName;
 
 	const FText DefaultToolTip = InArgs.DefaultGroupToolTip.IsEmpty()
-		? LOCTEXT("DefaultGroupTooltip", "Saved items.")
+		? LOCTEXT("DefaultGroupTooltip", "Pinned items.")
 		: InArgs.DefaultGroupToolTip;
 
 	Groups.Add(MakeShared<FMediaViewerLibraryGroup>(
diff --git a/Engine/Plugins/Experimental/MediaViewer/Source/MediaViewer/Private/MediaViewerCommands.cpp b/Engine/Plugins/Experimental/MediaViewer/Source/MediaViewer/Private/MediaViewerCommands.cpp
--- a/Engine/Plugins/Experimental/MediaViewer/Source/MediaViewer/Private/MediaViewerCommands.cpp
+++ b/Engine/Plugins/Experimental/MediaViewer/Source/MediaViewer/Private/MediaViewerCommands.cpp
@@ -267,8 +267,8 @@ void FMediaViewerCommands::RegisterCommands()
 
 	UI_COMMAND(
 		AddToLibrary,
-		"Add to Library",
-		"Save this image to the default group in the Library.",
+		"Pin Image",
+		"Add this image to the Pinned group in the Library.",
 		EUserInterfaceActionType::Button,
 		FInputChord(EKeys::S, EModifierKey::Control)
 	);
diff --git a/Engine/Plugins/Experimental/MediaViewer/Source/MediaViewer/Private/Widgets/SMediaImageViewerOverlay.cpp b/Engine/Plugins/Experimental/MediaViewer/Source/MediaViewer/Private/Widgets/SMediaImageViewerOverlay.cpp
--- a/Engine/Plugins/Experimental/MediaViewer/Source/MediaViewer/Private/Widgets/SMediaImageViewerOverlay.cpp
+++ b/Engine/Plugins/Experimental/MediaViewer/Source/MediaViewer/Private/Widgets/SMediaImageViewerOverlay.cpp
@@ -888,9 +888,7 @@ bool SMediaImageViewerOverlay::CanAddToLibrary() const
 		return false;
 	}
 
-	TSharedRef<FMediaViewerLibrary> Library = Delegates->GetLibrary.Execute();
-
-	return !Library->FindItemByValue(CachedItem->GetItemType(), CachedItem->GetStringValue()).IsValid();
+	return !Delegates->GetLibrary.Execute()->GetItemGroup(CachedItem->GetId()).IsValid();
 }
 
 void SMediaImageViewerOverlay::AddToLibrary()
