Media Viewer: Fixed crash relating for viewport.

diff --git a/Engine/Plugins/Experimental/MediaViewer/Source/MediaViewer/Private/ImageViewers/LevelEditorViewportImageViewer.cpp b/Engine/Plugins/Experimental/MediaViewer/Source/MediaViewer/Private/ImageViewers/LevelEditorViewportImageViewer.cpp
--- a/Engine/Plugins/Experimental/MediaViewer/Source/MediaViewer/Private/ImageViewers/LevelEditorViewportImageViewer.cpp
+++ b/Engine/Plugins/Experimental/MediaViewer/Source/MediaViewer/Private/ImageViewers/LevelEditorViewportImageViewer.cpp
@@ -170,7 +170,7 @@ UTextureRenderTarget2D* FLevelEditorViewportImageViewer::FItem::CreateRenderTarg
 		TNotNull<FViewport*>(Viewport.Get()),
 		TNotNull<UTextureRenderTarget2D*>(RenderTarget),
 		FRenderComplete(),
-		/* Resize render target */ false
+		/* Resize render target */ true
 	);
 
 	return RenderTarget;
diff --git a/Engine/Plugins/Experimental/MediaViewer/Source/MediaViewer/Private/ImageViewers/ViewportImageViewer.cpp b/Engine/Plugins/Experimental/MediaViewer/Source/MediaViewer/Private/ImageViewers/ViewportImageViewer.cpp
--- a/Engine/Plugins/Experimental/MediaViewer/Source/MediaViewer/Private/ImageViewers/ViewportImageViewer.cpp
+++ b/Engine/Plugins/Experimental/MediaViewer/Source/MediaViewer/Private/ImageViewers/ViewportImageViewer.cpp
@@ -79,7 +79,7 @@ void FViewportImageViewer::PaintImage(FMediaImageSlatePaintParams& InPaintParams
 {
 	if (RenderState == ERenderState::NotRendered || ViewportSettings.bRealTime)
 	{
-		RequestRender(/* Update render target size */ false);
+		RequestRender(/* Update render target size */ true);
 	}
 
 	FMediaImageViewer::PaintImage(InPaintParams, InPaintGeometry);
