Media Viewer: Transform lock now also duplicates media playback controls between players

diff --git a/Engine/Plugins/Experimental/MediaViewer/Source/MediaViewer/Private/Widgets/MediaViewerDelegates.h b/Engine/Plugins/Experimental/MediaViewer/Source/MediaViewer/Private/Widgets/MediaViewerDelegates.h
--- a/Engine/Plugins/Experimental/MediaViewer/Source/MediaViewer/Private/Widgets/MediaViewerDelegates.h
+++ b/Engine/Plugins/Experimental/MediaViewer/Source/MediaViewer/Private/Widgets/MediaViewerDelegates.h
@@ -25,6 +25,12 @@ namespace UE::MediaViewer::Private
 
 class FMediaViewerLibrary;
 
+struct FMediaImageViewerEventParams
+{
+	TSharedPtr<FMediaImageViewer> FromViewer;
+	FName EventName;
+};
+
 struct FMediaViewerDelegates
 {
 	DECLARE_DELEGATE_RetVal(FVector2D, FGetLocation)
@@ -49,6 +55,7 @@ struct FMediaViewerDelegates
 	DECLARE_DELEGATE_RetVal_OneParam(FIntPoint, FGetLocationForPosition, EMediaImageViewerPosition)
 	DECLARE_DELEGATE_RetVal_OneParam(bool, FGetBoolForPosition, EMediaImageViewerPosition)
 	DECLARE_DELEGATE_RetVal_OneParam(TSharedPtr<FUICommandList>, FGetCommandListForPosition, EMediaImageViewerPosition)
+	DECLARE_MULTICAST_DELEGATE_OneParam(FImageViewerEvent, const FMediaImageViewerEventParams&)
 
 	//////////////////////////////
 	// General viewer delegates //
@@ -159,6 +166,9 @@ struct FMediaViewerDelegates
 
 	/** Gets the command list for image viewer in the given position. */
 	FGetCommandListForPosition GetCommandListForPosition;
+
+	/** Generic event that can be broadcast between image viewers and controls. */
+	FImageViewerEvent ImageViewerEvent;
 };
 
 } // UE::MediaViewer::Private
diff --git a/Engine/Plugins/Experimental/MediaViewer/Source/MediaViewer/Private/Widgets/SMediaSourceOverlay.cpp b/Engine/Plugins/Experimental/MediaViewer/Source/MediaViewer/Private/Widgets/SMediaSourceOverlay.cpp
--- a/Engine/Plugins/Experimental/MediaViewer/Source/MediaViewer/Private/Widgets/SMediaSourceOverlay.cpp
+++ b/Engine/Plugins/Experimental/MediaViewer/Source/MediaViewer/Private/Widgets/SMediaSourceOverlay.cpp
@@ -34,6 +34,17 @@ namespace UE::MediaViewer::Private
 
 constexpr double FadeTime = 2.0;
 
+namespace MediaSourceOverlay::Events
+{
+	const FLazyName Rewind      = "MediaSource_Rewind";
+	const FLazyName Reverse     = "MediaSource_Reverse";
+	const FLazyName StepBack    = "MediaSource_StepBack";
+	const FLazyName Play        = "MediaSource_Play";
+	const FLazyName Pause       = "MediaSource_Pause";
+	const FLazyName StepForward = "MediaSource_StepForward";
+	const FLazyName Forward     = "MediaSource_Forward";
+}
+
 void SMediaSourceOverlay::PrivateRegisterAttributes(FSlateAttributeDescriptor::FInitializer&)
 {
 }
@@ -46,6 +57,7 @@ void SMediaSourceOverlay::Construct(const FArguments& InArgs, const TSharedRef<F
 	if (InViewerTab.IsValid())
 	{
 		Delegates = InViewerTab->GetViewer()->GetDelegates();
+		Delegates->ImageViewerEvent.AddSP(this, &SMediaSourceOverlay::ReceiveEvent);
 	}
 
 	TrySetFrameRate();
@@ -492,6 +504,8 @@ FReply SMediaSourceOverlay::Rewind_OnClicked()
 	{
 		MediaPlayer->Pause();
 		MediaPlayer->Rewind();
+
+		SendEvent(MediaSourceOverlay::Events::Rewind);
 	}
 
 	return FReply::Handled();
@@ -545,11 +559,17 @@ FReply SMediaSourceOverlay::Reverse_OnClicked()
 
 		if (FMath::IsNearlyZero(Rate))
 		{
-			MediaPlayer->SetRate(-1.f);
+			if (Command_Reverse(MediaPlayer))
+			{
+				SendEvent(MediaSourceOverlay::Events::Reverse);
+			}
 		}
 		else
 		{
-			MediaPlayer->Pause();
+			if (Command_Pause(MediaPlayer))
+			{
+				SendEvent(MediaSourceOverlay::Events::Pause);
+			}
 		}
 	}
 
@@ -572,11 +592,9 @@ FReply SMediaSourceOverlay::StepBack_OnClicked()
 {
 	if (UMediaPlayer* MediaPlayer = GetMediaPlayer())
 	{
-		const float FrameRate = MediaPlayer->GetVideoTrackFrameRate(INDEX_NONE, INDEX_NONE);
-
-		if (!FMath::IsNearlyZero(FrameRate))
+		if (Command_StepBack(MediaPlayer))
 		{
-			MediaPlayer->Seek(MediaPlayer->GetTime() - FTimespan::FromSeconds(1.f / FrameRate));
+			SendEvent(MediaSourceOverlay::Events::StepBack);
 		}
 	}
 
@@ -631,11 +649,17 @@ FReply SMediaSourceOverlay::Play_OnClicked()
 
 		if (FMath::IsNearlyZero(Rate))
 		{
-			MediaPlayer->Play();
+			if (Command_Play(MediaPlayer))
+			{
+				SendEvent(MediaSourceOverlay::Events::Play);
+			}
 		}
 		else
 		{
-			MediaPlayer->Pause();
+			if (Command_Pause(MediaPlayer))
+			{
+				SendEvent(MediaSourceOverlay::Events::Pause);
+			}
 		}
 	}
 
@@ -658,11 +682,9 @@ FReply SMediaSourceOverlay::StepForward_OnClicked()
 {
 	if (UMediaPlayer* MediaPlayer = GetMediaPlayer())
 	{
-		const float FrameRate = MediaPlayer->GetVideoTrackFrameRate(INDEX_NONE, INDEX_NONE);
-
-		if (!FMath::IsNearlyZero(FrameRate))
+		if (Command_StepForward(MediaPlayer))
 		{
-			MediaPlayer->Seek(MediaPlayer->GetTime() + FTimespan::FromSeconds(1.f / FrameRate));
+			SendEvent(MediaSourceOverlay::Events::StepForward);
 		}
 	}
 
@@ -683,24 +705,165 @@ FReply SMediaSourceOverlay::Forward_OnClicked()
 {
 	if (UMediaPlayer* MediaPlayer = GetMediaPlayer())
 	{
-		// INDEX_NONE here access the currently playing TrackIndex and FormatIndex.
-		const double FrameRate = MediaPlayer->GetVideoTrackFrameRate(INDEX_NONE, INDEX_NONE);
-
-		if (!FMath::IsNearlyZero(FrameRate))
+		if (Command_Forward(MediaPlayer))
 		{
-			MediaPlayer->Pause();
-
-			const double FrameTime = 1.0 / FrameRate;
-			FTimespan SeekLocation = MediaPlayer->GetDuration();
-			SeekLocation -= FTimespan::FromSeconds(FrameTime);
-
-			MediaPlayer->Seek(SeekLocation);
+			SendEvent(MediaSourceOverlay::Events::Forward);
 		}
 	}
 
 	return FReply::Handled();
 }
 
+void SMediaSourceOverlay::SendEvent(FName InEventName)
+{
+	if (!Delegates.IsValid())
+	{
+		return;
+	}
+
+	if (!Delegates->AreTransformsLocked.Execute())
+	{
+		return;
+	}
+
+	FMediaImageViewerEventParams Params;
+	Params.FromViewer = ImageViewerWeak.Pin();
+	Params.EventName = InEventName;
+
+	Delegates->ImageViewerEvent.Broadcast(Params);
+}
+
+void SMediaSourceOverlay::ReceiveEvent(const FMediaImageViewerEventParams& InEventParams)
+{
+	TSharedPtr<FMediaSourceImageViewer> ImageViewer = ImageViewerWeak.Pin();
+
+	if (!ImageViewer.IsValid())
+	{
+		return;
+	}
+
+	if (InEventParams.FromViewer.Get() == ImageViewer.Get())
+	{
+		return;
+	}
+
+	UMediaPlayer* MediaPlayer = GetMediaPlayer();
+
+	if (!MediaPlayer)
+	{
+		return;
+	}
+
+	using namespace MediaSourceOverlay;
+
+	if (InEventParams.EventName == Events::Rewind)
+	{
+		Command_Rewind(MediaPlayer);
+	}
+	else if (InEventParams.EventName == Events::Reverse)
+	{
+		Command_Reverse(MediaPlayer);
+	}
+	else if (InEventParams.EventName == Events::StepBack)
+	{
+		Command_StepBack(MediaPlayer);
+	}
+	else if (InEventParams.EventName == Events::Play)
+	{
+		Command_Play(MediaPlayer);
+	}
+	else if (InEventParams.EventName == Events::Pause)
+	{
+		Command_Pause(MediaPlayer);
+	}
+	else if (InEventParams.EventName == Events::StepForward)
+	{
+		Command_StepForward(MediaPlayer);
+	}
+	else if (InEventParams.EventName == Events::Forward)
+	{
+		Command_Forward(MediaPlayer);
+	}
+}
+
+bool SMediaSourceOverlay::Command_Rewind(UMediaPlayer* InMediaPlayer)
+{
+	InMediaPlayer->Pause();
+	InMediaPlayer->Rewind();
+
+	return true;
+}
+
+bool SMediaSourceOverlay::Command_Reverse(UMediaPlayer* InMediaPlayer)
+{
+	InMediaPlayer->SetRate(-1.f);
+
+	return true;
+}
+
+bool SMediaSourceOverlay::Command_StepBack(UMediaPlayer* InMediaPlayer)
+{
+	const float FrameRate = InMediaPlayer->GetVideoTrackFrameRate(INDEX_NONE, INDEX_NONE);
+
+	if (!FMath::IsNearlyZero(FrameRate))
+	{
+		InMediaPlayer->Seek(InMediaPlayer->GetTime() - FTimespan::FromSeconds(1.f / FrameRate));
+
+		return true;
+	}
+
+	return false;
+}
+
+bool SMediaSourceOverlay::Command_Play(UMediaPlayer* InMediaPlayer)
+{
+	InMediaPlayer->Play();
+
+	return true;
+}
+
+bool SMediaSourceOverlay::Command_Pause(UMediaPlayer* InMediaPlayer)
+{
+	InMediaPlayer->Pause();
+
+	return true;
+}
+
+bool SMediaSourceOverlay::Command_StepForward(UMediaPlayer* InMediaPlayer)
+{
+	const float FrameRate = InMediaPlayer->GetVideoTrackFrameRate(INDEX_NONE, INDEX_NONE);
+
+	if (!FMath::IsNearlyZero(FrameRate))
+	{
+		InMediaPlayer->Seek(InMediaPlayer->GetTime() + FTimespan::FromSeconds(1.f / FrameRate));
+
+		return true;
+	}
+
+	return false;
+}
+
+bool SMediaSourceOverlay::Command_Forward(UMediaPlayer* InMediaPlayer)
+{
+	// INDEX_NONE here access the currently playing TrackIndex and FormatIndex.
+	const double FrameRate = InMediaPlayer->GetVideoTrackFrameRate(INDEX_NONE, INDEX_NONE);
+
+	if (!FMath::IsNearlyZero(FrameRate))
+	{
+		InMediaPlayer->Pause();
+
+		const double FrameTime = 1.0 / FrameRate;
+		FTimespan SeekLocation = InMediaPlayer->GetDuration();
+		SeekLocation -= FTimespan::FromSeconds(FrameTime);
+
+		InMediaPlayer->Seek(SeekLocation);
+
+		return true;
+	}
+
+	return false;
+}
+
 }
 
 #undef LOCTEXT_NAMESPACE
diff --git a/Engine/Plugins/Experimental/MediaViewer/Source/MediaViewer/Private/Widgets/SMediaSourceOverlay.h b/Engine/Plugins/Experimental/MediaViewer/Source/MediaViewer/Private/Widgets/SMediaSourceOverlay.h
--- a/Engine/Plugins/Experimental/MediaViewer/Source/MediaViewer/Private/Widgets/SMediaSourceOverlay.h
+++ b/Engine/Plugins/Experimental/MediaViewer/Source/MediaViewer/Private/Widgets/SMediaSourceOverlay.h
@@ -22,6 +22,7 @@ namespace UE::MediaViewer::Private
 {
 
 class FMediaSourceImageViewer;
+struct FMediaImageViewerEventParams;
 struct FMediaViewerDelegates;
 	
 class SMediaSourceOverlay : public SCompoundWidget
@@ -92,6 +93,18 @@ protected:
 
 	bool Forward_IsEnabled() const;
 	FReply Forward_OnClicked();
+
+	void SendEvent(FName InEventName);
+
+	void ReceiveEvent(const FMediaImageViewerEventParams& InEventParams);
+
+	static bool Command_Rewind(UMediaPlayer* InMediaPlayer);
+	static bool Command_Reverse(UMediaPlayer* InMediaPlayer);
+	static bool Command_StepBack(UMediaPlayer* InMediaPlayer);
+	static bool Command_Play(UMediaPlayer* InMediaPlayer);
+	static bool Command_Pause(UMediaPlayer* InMediaPlayer);
+	static bool Command_StepForward(UMediaPlayer* InMediaPlayer);
+	static bool Command_Forward(UMediaPlayer* InMediaPlayer);
 };
 
 } // UE::MediaViewer::Private
\ No newline at end of file
