Remote Control: Small bug fix

diff --git a/Engine/Plugins/Experimental/MediaViewer/Source/MediaViewer/Private/Widgets/SMediaSourceOverlay.cpp b/Engine/Plugins/Experimental/MediaViewer/Source/MediaViewer/Private/Widgets/SMediaSourceOverlay.cpp
--- a/Engine/Plugins/Experimental/MediaViewer/Source/MediaViewer/Private/Widgets/SMediaSourceOverlay.cpp
+++ b/Engine/Plugins/Experimental/MediaViewer/Source/MediaViewer/Private/Widgets/SMediaSourceOverlay.cpp
@@ -721,7 +721,7 @@ void SMediaSourceOverlay::SendEvent(FName InEventName)
 		return;
 	}
 
-	if (!Delegates->AreTransformsLocked.Execute())
+	if (!Delegates->AreTransformsLocked.IsBound() || !Delegates->AreTransformsLocked.Execute())
 	{
 		return;
 	}
