Media Viewer: Added analytics to the tab opening.

diff --git a/Engine/Plugins/Experimental/MediaViewer/Source/MediaViewer/Private/Widgets/SMediaViewerTab.cpp b/Engine/Plugins/Experimental/MediaViewer/Source/MediaViewer/Private/Widgets/SMediaViewerTab.cpp
--- a/Engine/Plugins/Experimental/MediaViewer/Source/MediaViewer/Private/Widgets/SMediaViewerTab.cpp
+++ b/Engine/Plugins/Experimental/MediaViewer/Source/MediaViewer/Private/Widgets/SMediaViewerTab.cpp
@@ -2,6 +2,7 @@
 
 #include "Widgets/SMediaViewerTab.h"
 
+#include "EngineAnalytics.h"
 #include "ImageViewers/ColorImageViewer.h"
 #include "ImageViewers/NullImageViewer.h"
 #include "Library/MediaViewerLibrary.h"
@@ -40,6 +41,11 @@ void SMediaViewerTab::Construct(const FArguments& InArgs, const FMediaViewerArgs
 	[
 		SAssignNew(Viewer, Private::SMediaViewer, SharedThis(this), InMediaViewerArgs, ImageViewerLeft.ToSharedRef(), ImageViewerRight.ToSharedRef())
 	];
+
+	if (FEngineAnalytics::IsAvailable())
+	{
+		FEngineAnalytics::GetProvider().RecordEvent(TEXT("Editor.Usage.MediaViewer"), TEXT("Action"), TEXT("OpenedTab"));
+	}
 }
 
 const FMediaViewerArgs& SMediaViewerTab::GetArgs() const
