Media Viewer: Snapshot addendum
- Fixed UV range for images on snapshots (canvas vs image UV)
- Added names of images to snapshot default filename

diff --git a/Engine/Plugins/Experimental/MediaViewer/Source/MediaViewer/Private/ImageViewers/MediaImageViewer.cpp b/Engine/Plugins/Experimental/MediaViewer/Source/MediaViewer/Private/ImageViewers/MediaImageViewer.cpp
--- a/Engine/Plugins/Experimental/MediaViewer/Source/MediaViewer/Private/ImageViewers/MediaImageViewer.cpp
+++ b/Engine/Plugins/Experimental/MediaViewer/Source/MediaViewer/Private/ImageViewers/MediaImageViewer.cpp
@@ -280,14 +280,36 @@ void FMediaImageViewer::PaintImage(const FMediaImageCanvasPaintParams& InPaintPa
 		return;
 	}
 
+	const FVector2D CanvasUVRange = FVector2D(0, 1);
+
 	switch (InPaintParams.Orientation)
 	{
 		case EOrientation::Orient_Horizontal:
 		{
-			const float PositionX = InPaintGeometry.Position.X + (InPaintParams.UVRange.GetLowerBound().GetValue() * InPaintGeometry.Size.X);
-			const float SizeX = InPaintGeometry.Size.X * InPaintParams.UVRange.Size<float>();
-			const float U = InPaintParams.UVRange.GetLowerBound().GetValue() * ForegroundTexture->GetSurfaceWidth();
-			const float UL = ForegroundTexture->GetSurfaceWidth() * InPaintParams.UVRange.Size<float>();
+			const FVector2D ImageCanvasUVRange = FVector2D(
+				-InPaintGeometry.Position.X / InPaintGeometry.Size.X,
+				(InPaintParams.Canvas->SizeX - InPaintGeometry.Position.X) / InPaintGeometry.Size.X
+			);
+
+			// UV Range is based on the canvas, not the image. Convert to image range and clamp at 0-1.
+			const float ImageUVRangeMinX = FMath::Clamp(FMath::GetMappedRangeValueUnclamped(
+				CanvasUVRange,
+				ImageCanvasUVRange,
+				InPaintParams.UVRange.GetLowerBound().GetValue()
+			), 0, 1);
+
+			const float ImageUVRangeMaxX = FMath::Clamp(FMath::GetMappedRangeValueUnclamped(
+				CanvasUVRange,
+				ImageCanvasUVRange,
+				InPaintParams.UVRange.GetUpperBound().GetValue()
+			), 0, 1);
+
+			const float ImageUVRangeSize = ImageUVRangeMaxX - ImageUVRangeMinX;
+
+			const float PositionX = InPaintGeometry.Position.X + (ImageUVRangeMinX * InPaintGeometry.Size.X);
+			const float SizeX = InPaintGeometry.Size.X * ImageUVRangeSize;
+			const float U = ImageUVRangeMinX * ForegroundTexture->GetSurfaceWidth();
+			const float UL = ForegroundTexture->GetSurfaceWidth() * ImageUVRangeSize;
 
 			InPaintParams.Canvas->DrawTile(
 				ForegroundTexture,
@@ -306,10 +328,30 @@ void FMediaImageViewer::PaintImage(const FMediaImageCanvasPaintParams& InPaintPa
 
 		case EOrientation::Orient_Vertical:
 		{
-			const float PositionY = InPaintGeometry.Position.Y + (InPaintParams.UVRange.GetLowerBound().GetValue() * InPaintGeometry.Size.Y);
-			const float SizeY = InPaintGeometry.Size.Y * InPaintParams.UVRange.Size<float>();
-			const float V = InPaintParams.UVRange.GetLowerBound().GetValue() * ForegroundTexture->GetSurfaceHeight();
-			const float VL = ForegroundTexture->GetSurfaceHeight() * InPaintParams.UVRange.Size<float>();
+			const FVector2D ImageCanvasUVRange = FVector2D(
+				-InPaintGeometry.Position.Y / InPaintGeometry.Size.Y,
+				(InPaintParams.Canvas->SizeY - InPaintGeometry.Position.Y) / InPaintGeometry.Size.Y
+			);
+
+			// UV Range is based on the canvas, not the image. Convert to image range and clamp at 0-1.
+			const float ImageUVRangeMinY = FMath::Clamp(FMath::GetMappedRangeValueUnclamped(
+				CanvasUVRange,
+				ImageCanvasUVRange,
+				InPaintParams.UVRange.GetLowerBound().GetValue()
+			), 0, 1);
+
+			const float ImageUVRangeMaxY = FMath::Clamp(FMath::GetMappedRangeValueUnclamped(
+				CanvasUVRange,
+				ImageCanvasUVRange,
+				InPaintParams.UVRange.GetUpperBound().GetValue()
+			), 0, 1);
+
+			const float ImageUVRangeSize = ImageUVRangeMaxY - ImageUVRangeMinY;
+
+			const float PositionY = InPaintGeometry.Position.Y + (ImageUVRangeMinY * InPaintGeometry.Size.Y);
+			const float SizeY = InPaintGeometry.Size.Y * ImageUVRangeSize;
+			const float V = ImageUVRangeMinY * ForegroundTexture->GetSurfaceHeight();
+			const float VL = ForegroundTexture->GetSurfaceHeight() * ImageUVRangeSize;
 
 			InPaintParams.Canvas->DrawTile(
 				ForegroundTexture,
diff --git a/Engine/Plugins/Experimental/MediaViewer/Source/MediaViewer/Private/Widgets/SMediaViewer.cpp b/Engine/Plugins/Experimental/MediaViewer/Source/MediaViewer/Private/Widgets/SMediaViewer.cpp
--- a/Engine/Plugins/Experimental/MediaViewer/Source/MediaViewer/Private/Widgets/SMediaViewer.cpp
+++ b/Engine/Plugins/Experimental/MediaViewer/Source/MediaViewer/Private/Widgets/SMediaViewer.cpp
@@ -1092,6 +1092,15 @@ void SMediaViewer::TakeSnapshot()
 	SaveAssetDialogConfig.ExistingAssetPolicy = ESaveAssetDialogExistingAssetPolicy::AllowButWarn;
 	SaveAssetDialogConfig.DefaultAssetName = LOCTEXT("MediaViewerSnapshot", "MediaViewerSnapshot").ToString();
 
+	for (int32 Index = 0; Index < static_cast<int32>(EMediaImageViewerPosition::COUNT); ++Index)
+	{
+		if (TSharedPtr<FMediaImageViewer> ImageViewer = GetImageViewer(static_cast<EMediaImageViewerPosition>(Index)))
+		{
+			SaveAssetDialogConfig.DefaultAssetName += TEXT("_");
+			SaveAssetDialogConfig.DefaultAssetName += ImageViewer->GetInfo().DisplayName.ToString();
+		}
+	}
+
 	FContentBrowserModule& ContentBrowserModule = FModuleManager::LoadModuleChecked<FContentBrowserModule>("ContentBrowser");
 	FString SaveObjectPath = ContentBrowserModule.Get().CreateModalSaveAssetDialog(SaveAssetDialogConfig);
 
