Media Viewer: Now beta!

diff --git a/Engine/Plugins/Experimental/MediaViewer/MediaViewer.uplugin b/Engine/Plugins/Experimental/MediaViewer/MediaViewer.uplugin
--- a/Engine/Plugins/Experimental/MediaViewer/MediaViewer.uplugin
+++ b/Engine/Plugins/Experimental/MediaViewer/MediaViewer.uplugin
@@ -6,8 +6,8 @@
 	"Category": "Other",
 	"CreatedBy": "Epic Games, Inc.",
 	"CreatedByURL": "https://epicgames.com",
-	"IsBetaVersion": false,
-	"IsExperimentalVersion": true,
+	"IsBetaVersion": true,
+	"IsExperimentalVersion": false,
 	"CanContainContent": true,
 	"Modules": [
 		{
