Media Viewer: New dragging options
- Revamped dragging system a little to support more zoom types.
- Added alt+rmb to zoom (in addition to scroll wheel)
- Added alt+mmb for panning (in addition to left click)

diff --git a/Engine/Plugins/Experimental/MediaViewer/Source/MediaViewer/Private/Widgets/SMediaImageViewerOverlay.cpp b/Engine/Plugins/Experimental/MediaViewer/Source/MediaViewer/Private/Widgets/SMediaImageViewerOverlay.cpp
--- a/Engine/Plugins/Experimental/MediaViewer/Source/MediaViewer/Private/Widgets/SMediaImageViewerOverlay.cpp
+++ b/Engine/Plugins/Experimental/MediaViewer/Source/MediaViewer/Private/Widgets/SMediaImageViewerOverlay.cpp
@@ -34,8 +34,7 @@ const FName DropTargetName = TEXT("DropTarget");
 SMediaImageViewerOverlay::SMediaImageViewerOverlay()
 	: CommandList(MakeShared<FUICommandList>())
 	, DragButtonName(NAME_None)
-	, bDragging(false)
-	, bExternalDragging(false)
+	, DragType(EDragType::None)
 	, bOverlayEnabled(true)
 {	
 }
@@ -161,7 +160,7 @@ bool SMediaImageViewerOverlay::IsCursorOverImageViewer() const
 
 FCursorReply SMediaImageViewerOverlay::OnCursorQuery(const FGeometry& InMyGeometry, const FPointerEvent& InCursorEvent) const
 {
-	if (bDragging || bExternalDragging)
+	if (DragType != EDragType::None)
 	{
 		return FCursorReply::Cursor(EMouseCursor::GrabHandClosed);
 	}
@@ -186,16 +185,25 @@ FReply SMediaImageViewerOverlay::OnKeyDown(const FGeometry& InMyGeometry, const
 
 FReply SMediaImageViewerOverlay::OnMouseButtonDown(const FGeometry& InMyGeometry, const FPointerEvent& InMouseEvent)
 {
-	if (!bDragging && !bExternalDragging)
+	if (DragType == EDragType::None)
 	{
 		const FVector2D CursorLocation = Delegates->GetCursorLocation.Execute();
 
-		if (InMouseEvent.GetEffectingButton() == EKeys::LeftMouseButton)
+		if (InMouseEvent.GetEffectingButton() == EKeys::LeftMouseButton
+			|| (InMouseEvent.GetEffectingButton() == EKeys::MiddleMouseButton && InMouseEvent.IsAltDown()))
 		{
 			DragStartCursor = CursorLocation;
 			DragStartOffset = GetOffset();
 			DragButtonName = InMouseEvent.GetEffectingButton().GetFName();
-			bDragging = true;
+			DragType = EDragType::Offset;
+			return FReply::Handled();
+		}
+		else if (InMouseEvent.GetEffectingButton() == EKeys::RightMouseButton && InMouseEvent.IsAltDown())
+		{
+			DragStartCursor = CursorLocation;
+			LastDragScaleValue = GetScale();
+			DragButtonName = InMouseEvent.GetEffectingButton().GetFName();
+			DragType = EDragType::Scale;
 			return FReply::Handled();
 		}
 
@@ -206,7 +214,7 @@ FReply SMediaImageViewerOverlay::OnMouseButtonDown(const FGeometry& InMyGeometry
 			if (ImageViewer->OnTrackingStarted(EventState, FIntPoint(FMath::RoundToInt(CursorLocation.X), FMath::RoundToInt(CursorLocation.Y))))
 			{
 				DragButtonName = InMouseEvent.GetEffectingButton().GetFName();
-				bExternalDragging = true;
+				DragType = EDragType::External;
 				return FReply::Handled();
 			}
 		}
@@ -217,7 +225,7 @@ FReply SMediaImageViewerOverlay::OnMouseButtonDown(const FGeometry& InMyGeometry
 
 FReply SMediaImageViewerOverlay::OnMouseButtonUp(const FGeometry& InMyGeometry, const FPointerEvent& InMouseEvent)
 {
-	if (!bDragging && !bExternalDragging)
+	if (DragType == EDragType::None)
 	{
 		if (InMouseEvent.GetEffectingButton() == EKeys::RightMouseButton)
 		{
@@ -247,17 +255,16 @@ FReply SMediaImageViewerOverlay::OnMouseWheel(const FGeometry& InMyGeometry, con
 	static const float SlowMultiplier = FMath::Pow(2.0, 0.03125);
 
 	const FModifierKeysState ModifierKeys = FSlateApplication::Get().GetModifierKeys();
-	
-	const float ModifierMultiplier = 
-		ModifierKeys.IsLeftShiftDown() 
+
+	const float ModifierScale = ModifierKeys.IsLeftShiftDown()
 		? FastMultiplier
 		: ((ModifierKeys.IsLeftControlDown() || ModifierKeys.IsLeftCommandDown())
 			? SlowMultiplier
 			: NormalMultiplier);
-	
+
 	const float ScaleMultiplier = InMouseEvent.GetWheelDelta() > 0
-		? ModifierMultiplier
-		: (1.f / ModifierMultiplier);
+		? ModifierScale
+		: (1.f / ModifierScale);
 
 	if (Delegates->AreTransformsLocked.Execute())
 	{
@@ -309,9 +316,19 @@ void SMediaImageViewerOverlay::UpdateMouse(const TOptional<FVector2D>& InMousePo
 		}
 	}
 
-	if (bDragging)
+	switch (DragType)
 	{
-		UpdateDragPosition();
+		case EDragType::Offset:
+			UpdateDragPosition();
+			break;
+
+		case EDragType::Scale:
+			UpdateDragScale();
+			break;
+
+		case EDragType::External:
+			UpdateDragExternal();
+			break;
 	}
 
 	if (InMousePosition.IsSet())
@@ -330,23 +347,23 @@ void SMediaImageViewerOverlay::OnDragButtonUp(FName InKeyName)
 		return;
 	}
 
-	if (bDragging)
-	{
-		UpdateDragPosition();
-		DragButtonName = NAME_None;
-		bDragging = false;
-	}
-	else if (bExternalDragging)
+	switch (DragType)
 	{
-		if (TSharedPtr<FMediaImageViewer> ImageViewer = Delegates->GetImageViewer.Execute(Position))
-		{
-			const FVector2D CursorLocation = Delegates->GetCursorLocation.Execute();
-			ImageViewer->OnTrackingStopped(FIntPoint(FMath::RoundToInt(CursorLocation.X), FMath::RoundToInt(CursorLocation.Y)));
-		}
+		case EDragType::Offset:
+			UpdateDragPosition();
+			break;
+
+		case EDragType::Scale:
+			UpdateDragScale();
+			break;
 
-		DragButtonName = NAME_None;
-		bExternalDragging = false;
+		case EDragType::External:
+			UpdateDragExternal();
+			break;
 	}
+
+	DragType = EDragType::None;
+	DragButtonName = NAME_None;
 }
 
 void SMediaImageViewerOverlay::UpdateDragPosition()
@@ -360,6 +377,41 @@ void SMediaImageViewerOverlay::UpdateDragPosition()
 	Try_AddOffset(NewOffset - GetOffset());
 }
 
+void SMediaImageViewerOverlay::UpdateDragScale()
+{
+	static const float DragMultiplier = 0.001f;
+	static const float NormalMultiplier = 1.f;
+	static const float FastMultiplier = 4.f;
+	static const float SlowMultiplier = 0.25f;
+
+	const FModifierKeysState ModifierKeys = FSlateApplication::Get().GetModifierKeys();
+
+	const float ModifierScale = ModifierKeys.IsLeftShiftDown()
+		? FastMultiplier
+		: ((ModifierKeys.IsLeftControlDown() || ModifierKeys.IsLeftCommandDown())
+			? SlowMultiplier
+			: NormalMultiplier);
+
+	const FVector2D CursorLocation = Delegates->GetCursorLocation.Execute();
+	const float Difference = (CursorLocation - DragStartCursor).Y;
+	const float NewScale = LastDragScaleValue - Difference * ModifierScale * DragMultiplier;
+	
+	// Since we can change speed of the scale change mid operation, reset this each update.
+	DragStartCursor = CursorLocation;
+	LastDragScaleValue = NewScale;
+
+	Try_SetScale(NewScale);
+}
+
+void SMediaImageViewerOverlay::UpdateDragExternal()
+{
+	if (TSharedPtr<FMediaImageViewer> ImageViewer = Delegates->GetImageViewer.Execute(Position))
+	{
+		const FVector2D CursorLocation = Delegates->GetCursorLocation.Execute();
+		ImageViewer->OnTrackingStopped(FIntPoint(FMath::RoundToInt(CursorLocation.X), FMath::RoundToInt(CursorLocation.Y)));
+	}
+}
+
 EVisibility SMediaImageViewerOverlay::GetDragDescriptionVisibility() const
 {
 	return FSlateApplication::Get().IsDragDropping()
@@ -708,7 +760,7 @@ void SMediaImageViewerOverlay::SetScale(float InScale)
 {
 	if (TSharedPtr<FMediaImageViewer> ImageViewer = Delegates->GetImageViewer.Execute(Position))
 	{
-		ImageViewer->GetPaintSettings().Scale = InScale;
+		ImageViewer->GetPaintSettings().Scale = FMath::Clamp(InScale, SMediaViewer::MinScale, SMediaViewer::MaxScale);
 	}
 }
 
diff --git a/Engine/Plugins/Experimental/MediaViewer/Source/MediaViewer/Private/Widgets/SMediaImageViewerOverlay.h b/Engine/Plugins/Experimental/MediaViewer/Source/MediaViewer/Private/Widgets/SMediaImageViewerOverlay.h
--- a/Engine/Plugins/Experimental/MediaViewer/Source/MediaViewer/Private/Widgets/SMediaImageViewerOverlay.h
+++ b/Engine/Plugins/Experimental/MediaViewer/Source/MediaViewer/Private/Widgets/SMediaImageViewerOverlay.h
@@ -79,18 +79,28 @@ public:
 	//~ End SWidget
 
 protected:
+	enum class EDragType : uint8
+	{
+		/** No drag operatio is active. */
+		None,
+		/** Whether the widget is performing an internal offset drag operation. */
+		Offset,
+		/** Whether the widget is performing an internal scale drag operation. */
+		Scale,
+		/** Whether an image viewer is performing its own drag operation. */
+		External
+	};
+
 	EMediaImageViewerPosition Position;
 	TSharedPtr<FMediaViewerDelegates> Delegates;
 	TSharedPtr<FMediaViewerLibraryItem> CachedItem;
 	TSharedRef<FUICommandList> CommandList;
 	/** Name of the button which is currently performing a drag operation. */
 	FName DragButtonName;
-	/** Whether the widget is performing an "internal" drag operation. */
-	bool bDragging;
-	/** Whether an image viewer is performing its own drag operation. */
-	bool bExternalDragging;
+	EDragType DragType;
 	FVector2D DragStartCursor;
 	FVector DragStartOffset;
+	float LastDragScaleValue;
 	/** For moving the image based on cursor location. */
 	FIntPoint LastCursorPosition;
 	TSharedPtr<SMediaImageViewerStatusBar> StatusBar;
@@ -108,6 +118,8 @@ protected:
 
 	void OnDragButtonUp(FName InKeyName);
 	void UpdateDragPosition();
+	void UpdateDragScale();
+	void UpdateDragExternal();
 
 	EVisibility GetDragDescriptionVisibility() const;
 
diff --git a/Engine/Plugins/Experimental/MediaViewer/Source/MediaViewer/Private/Widgets/SMediaViewer.cpp b/Engine/Plugins/Experimental/MediaViewer/Source/MediaViewer/Private/Widgets/SMediaViewer.cpp
--- a/Engine/Plugins/Experimental/MediaViewer/Source/MediaViewer/Private/Widgets/SMediaViewer.cpp
+++ b/Engine/Plugins/Experimental/MediaViewer/Source/MediaViewer/Private/Widgets/SMediaViewer.cpp
@@ -884,7 +884,7 @@ void SMediaViewer::MultiplyScaleToAll(float InMultiple)
 	{
 		if (ImageViewer.IsValid())
 		{
-			ImageViewer->GetPaintSettings().Scale *= InMultiple;
+			ImageViewer->GetPaintSettings().Scale = FMath::Clamp(ImageViewer->GetPaintSettings().Scale * InMultiple, SMediaViewer::MinScale, SMediaViewer::MaxScale);
 		}
 	}
 }
diff --git a/Engine/Plugins/Experimental/MediaViewer/Source/MediaViewer/Private/Widgets/SMediaViewer.h b/Engine/Plugins/Experimental/MediaViewer/Source/MediaViewer/Private/Widgets/SMediaViewer.h
--- a/Engine/Plugins/Experimental/MediaViewer/Source/MediaViewer/Private/Widgets/SMediaViewer.h
+++ b/Engine/Plugins/Experimental/MediaViewer/Source/MediaViewer/Private/Widgets/SMediaViewer.h
@@ -56,6 +56,9 @@ class SMediaViewer : public SCompoundWidget, public FGCObject
 	SLATE_END_ARGS()
 
 public:
+	static constexpr float MaxScale = 10.f;
+	static constexpr float MinScale = 0.01f;
+
 	SMediaViewer();
 
 	virtual ~SMediaViewer() override;
