Media Viewer: Invalid assets are removed from history when the tab is opened.

diff --git a/Engine/Plugins/Experimental/MediaViewer/Source/MediaViewer/Private/ImageViewers/MaterialInterfaceImageViewer.cpp b/Engine/Plugins/Experimental/MediaViewer/Source/MediaViewer/Private/ImageViewers/MaterialInterfaceImageViewer.cpp
--- a/Engine/Plugins/Experimental/MediaViewer/Source/MediaViewer/Private/ImageViewers/MaterialInterfaceImageViewer.cpp
+++ b/Engine/Plugins/Experimental/MediaViewer/Source/MediaViewer/Private/ImageViewers/MaterialInterfaceImageViewer.cpp
@@ -125,7 +125,14 @@ bool FMaterialInterfaceImageViewer::FFactory::SupportsItemType(FName InItemType)
 
 TSharedPtr<FMediaViewerLibraryItem> FMaterialInterfaceImageViewer::FFactory::CreateLibraryItem(const FMediaViewerLibraryItem& InSavedItem) const
 {
-	return MakeShared<FMaterialInterfaceImageViewer::FItem>(FPrivateToken(), InSavedItem);
+	FAssetData AssetData;
+
+	if (IAssetRegistry::Get()->TryGetAssetByObjectPath(InSavedItem.GetStringValue(), AssetData) == UE::AssetRegistry::EExists::Exists)
+	{
+		return MakeShared<FMaterialInterfaceImageViewer::FItem>(FPrivateToken(), InSavedItem);
+	}
+
+	return nullptr;
 }
 
 FMaterialInterfaceImageViewer::FItem::FItem(const FText& InName, const FText& InToolTip, bool bInTransient, TNotNull<UMaterialInterface*> InMaterial)
diff --git a/Engine/Plugins/Experimental/MediaViewer/Source/MediaViewer/Private/ImageViewers/MediaSourceImageViewer.cpp b/Engine/Plugins/Experimental/MediaViewer/Source/MediaViewer/Private/ImageViewers/MediaSourceImageViewer.cpp
--- a/Engine/Plugins/Experimental/MediaViewer/Source/MediaViewer/Private/ImageViewers/MediaSourceImageViewer.cpp
+++ b/Engine/Plugins/Experimental/MediaViewer/Source/MediaViewer/Private/ImageViewers/MediaSourceImageViewer.cpp
@@ -2,6 +2,7 @@
 
 #include "ImageViewers/MediaSourceImageViewer.h"
 
+#include "AssetRegistry/IAssetRegistry.h"
 #include "Brushes/SlateImageBrush.h"
 #include "DetailLayoutBuilder.h"
 #include "Engine/Texture2D.h"
@@ -108,7 +109,15 @@ bool FMediaSourceImageViewer::FFactory::SupportsItemType(FName InItemType) const
 
 TSharedPtr<FMediaViewerLibraryItem> FMediaSourceImageViewer::FFactory::CreateLibraryItem(const FMediaViewerLibraryItem& InSavedItem) const
 {
-	return MakeShared<FMediaSourceImageViewer::FAssetItem>(FPrivateToken(), InSavedItem);
+	FAssetData AssetData;
+
+	if (IAssetRegistry::Get()->TryGetAssetByObjectPath(InSavedItem.GetStringValue(), AssetData) == UE::AssetRegistry::EExists::Exists)
+	{
+		return MakeShared<FMediaSourceImageViewer::FAssetItem>(FPrivateToken(), InSavedItem);
+	}
+
+	return nullptr;
+	
 }
 
 FMediaSourceImageViewer::FItem::FItem(const FText& InName, const FText& InToolTip, bool bInTransient, const FString& InStringValue)
diff --git a/Engine/Plugins/Experimental/MediaViewer/Source/MediaViewer/Private/ImageViewers/MediaTextureImageViewer.cpp b/Engine/Plugins/Experimental/MediaViewer/Source/MediaViewer/Private/ImageViewers/MediaTextureImageViewer.cpp
--- a/Engine/Plugins/Experimental/MediaViewer/Source/MediaViewer/Private/ImageViewers/MediaTextureImageViewer.cpp
+++ b/Engine/Plugins/Experimental/MediaViewer/Source/MediaViewer/Private/ImageViewers/MediaTextureImageViewer.cpp
@@ -2,6 +2,7 @@
 
 #include "ImageViewers/MediaTextureImageViewer.h"
 
+#include "AssetRegistry/IAssetRegistry.h"
 #include "Brushes/SlateImageBrush.h"
 #include "Engine/TextureRenderTarget2D.h"
 #include "Materials/MaterialInstanceDynamic.h"
@@ -88,7 +89,14 @@ bool FMediaTextureImageViewer::FFactory::SupportsItemType(FName InItemType) cons
 
 TSharedPtr<FMediaViewerLibraryItem> FMediaTextureImageViewer::FFactory::CreateLibraryItem(const FMediaViewerLibraryItem& InSavedItem) const
 {
-	return MakeShared<FMediaTextureImageViewer::FItem>(FPrivateToken(), InSavedItem);
+	FAssetData AssetData;
+
+	if (IAssetRegistry::Get()->TryGetAssetByObjectPath(InSavedItem.GetStringValue(), AssetData) == UE::AssetRegistry::EExists::Exists)
+	{
+		return MakeShared<FMediaTextureImageViewer::FItem>(FPrivateToken(), InSavedItem);
+	}
+
+	return nullptr;
 }
 
 FMediaTextureImageViewer::FItem::FItem(const FText& InName, const FText& InToolTip, bool bInTransient, TNotNull<UMediaTexture*> InMediaTexture)
diff --git a/Engine/Plugins/Experimental/MediaViewer/Source/MediaViewer/Private/ImageViewers/Texture2DImageViewer.cpp b/Engine/Plugins/Experimental/MediaViewer/Source/MediaViewer/Private/ImageViewers/Texture2DImageViewer.cpp
--- a/Engine/Plugins/Experimental/MediaViewer/Source/MediaViewer/Private/ImageViewers/Texture2DImageViewer.cpp
+++ b/Engine/Plugins/Experimental/MediaViewer/Source/MediaViewer/Private/ImageViewers/Texture2DImageViewer.cpp
@@ -2,6 +2,7 @@
 
 #include "ImageViewers/Texture2DImageViewer.h"
 
+#include "AssetRegistry/IAssetRegistry.h"
 #include "Brushes/SlateImageBrush.h"
 #include "Engine/Texture.h"
 #include "Engine/Texture2D.h"
@@ -88,7 +89,14 @@ bool FTexture2DImageViewer::FFactory::SupportsItemType(FName InItemType) const
 
 TSharedPtr<FMediaViewerLibraryItem> FTexture2DImageViewer::FFactory::CreateLibraryItem(const FMediaViewerLibraryItem& InSavedItem) const
 {
-	return MakeShared<FTexture2DImageViewer::FItem>(FPrivateToken(), InSavedItem);
+	FAssetData AssetData;
+
+	if (IAssetRegistry::Get()->TryGetAssetByObjectPath(InSavedItem.GetStringValue(), AssetData) == UE::AssetRegistry::EExists::Exists)
+	{
+		return MakeShared<FTexture2DImageViewer::FItem>(FPrivateToken(), InSavedItem);
+	}
+
+	return nullptr;
 }
 
 FTexture2DImageViewer::FItem::FItem(const FText& InName, const FText& InToolTip, bool bInTransient, TNotNull<UTexture2D*> InTexture)
diff --git a/Engine/Plugins/Experimental/MediaViewer/Source/MediaViewer/Private/ImageViewers/TextureRenderTarget2DImageViewer.cpp b/Engine/Plugins/Experimental/MediaViewer/Source/MediaViewer/Private/ImageViewers/TextureRenderTarget2DImageViewer.cpp
--- a/Engine/Plugins/Experimental/MediaViewer/Source/MediaViewer/Private/ImageViewers/TextureRenderTarget2DImageViewer.cpp
+++ b/Engine/Plugins/Experimental/MediaViewer/Source/MediaViewer/Private/ImageViewers/TextureRenderTarget2DImageViewer.cpp
@@ -2,6 +2,7 @@
 
 #include "ImageViewers/TextureRenderTarget2DImageViewer.h"
 
+#include "AssetRegistry/IAssetRegistry.h"
 #include "Brushes/SlateImageBrush.h"
 #include "Engine/TextureRenderTarget2D.h"
 #include "Math/Vector2D.h"
@@ -85,7 +86,15 @@ bool FTextureRenderTarget2DImageViewer::FFactory::SupportsItemType(FName InItemT
 
 TSharedPtr<FMediaViewerLibraryItem> FTextureRenderTarget2DImageViewer::FFactory::CreateLibraryItem(const FMediaViewerLibraryItem& InSavedItem) const
 {
-	return MakeShared<FTextureRenderTarget2DImageViewer::FItem>(FPrivateToken(), InSavedItem);
+	// While we support in-memory render targets, anything "saved" must be an asset.
+	FAssetData AssetData;
+
+	if (IAssetRegistry::Get()->TryGetAssetByObjectPath(InSavedItem.GetStringValue(), AssetData) == UE::AssetRegistry::EExists::Exists)
+	{
+		return MakeShared<FTextureRenderTarget2DImageViewer::FItem>(FPrivateToken(), InSavedItem);
+	}
+
+	return nullptr;
 }
 
 FTextureRenderTarget2DImageViewer::FItem::FItem(const FText& InName, const FText& InToolTip, bool bInTransient, TNotNull<UTextureRenderTarget2D*> InRenderTarget)
