Media Viewer: Clear history button

diff --git a/Engine/Plugins/Experimental/MediaViewer/Source/MediaViewer/Private/Library/MediaViewerLibrary.cpp b/Engine/Plugins/Experimental/MediaViewer/Source/MediaViewer/Private/Library/MediaViewerLibrary.cpp
--- a/Engine/Plugins/Experimental/MediaViewer/Source/MediaViewer/Private/Library/MediaViewerLibrary.cpp
+++ b/Engine/Plugins/Experimental/MediaViewer/Source/MediaViewer/Private/Library/MediaViewerLibrary.cpp
@@ -125,6 +125,24 @@ TSharedPtr<FMediaViewerLibraryGroup> FMediaViewerLibrary::GetGroup(const FGuid&
 	return nullptr;
 }
 
+int32 FMediaViewerLibrary::EmptyGroup(const FGuid& InGroupId)
+{
+	TSharedPtr<FMediaViewerLibraryGroup> Group = GetGroup(InGroupId);
+
+	if (!Group.IsValid())
+	{
+		return 0;
+	}
+
+	const int32 Count = Group->GetItems().Num();
+
+	Group->Empty();
+
+	OnChanged(IMediaViewerLibrary::EChangeType::ItemRemoved);
+
+	return Count;
+}
+
 bool FMediaViewerLibrary::CanRemoveGroup(const FGuid& InGroupIdToRemove) const
 {
 	if (InGroupIdToRemove == GetDefaultGroupId() || InGroupIdToRemove == GetHistoryGroupId() || InGroupIdToRemove == GetSnapshotsGroupId())
diff --git a/Engine/Plugins/Experimental/MediaViewer/Source/MediaViewer/Private/Library/MediaViewerLibrary.h b/Engine/Plugins/Experimental/MediaViewer/Source/MediaViewer/Private/Library/MediaViewerLibrary.h
--- a/Engine/Plugins/Experimental/MediaViewer/Source/MediaViewer/Private/Library/MediaViewerLibrary.h
+++ b/Engine/Plugins/Experimental/MediaViewer/Source/MediaViewer/Private/Library/MediaViewerLibrary.h
@@ -40,6 +40,7 @@ public:
 	virtual const TArray<TSharedRef<const FMediaViewerLibraryGroup>>& GetGroups() const override;
 	virtual bool AddGroup(const TSharedRef<FMediaViewerLibraryGroup>& InNewGroup) override;
 	virtual TSharedPtr<FMediaViewerLibraryGroup> GetGroup(const FGuid& InGroupIdToRemove) const override;
+	virtual int32 EmptyGroup(const FGuid& InGroupId) override;
 	virtual void RemoveInvalidGroupItems(const FGuid& InGroup) override;
 	virtual bool CanRemoveGroup(const FGuid& InGroupIdToRemove) const override;
 	virtual TSharedPtr<FMediaViewerLibraryGroup> RemoveGroup(const FGuid& InGroupIdToRemove) override;
diff --git a/Engine/Plugins/Experimental/MediaViewer/Source/MediaViewer/Private/Library/MediaViewerLibraryGroup.cpp b/Engine/Plugins/Experimental/MediaViewer/Source/MediaViewer/Private/Library/MediaViewerLibraryGroup.cpp
--- a/Engine/Plugins/Experimental/MediaViewer/Source/MediaViewer/Private/Library/MediaViewerLibraryGroup.cpp
+++ b/Engine/Plugins/Experimental/MediaViewer/Source/MediaViewer/Private/Library/MediaViewerLibraryGroup.cpp
@@ -95,6 +95,11 @@ bool FMediaViewerLibraryGroup::RemoveItemAt(int32 InIndex)
 	return true;
 }
 
+void FMediaViewerLibraryGroup::Empty()
+{
+	Items.Empty();
+}
+
 bool FMediaViewerLibraryGroup::IsDynamic() const
 {
 	return bDynamic;
diff --git a/Engine/Plugins/Experimental/MediaViewer/Source/MediaViewer/Private/Widgets/SMediaViewerLibraryGroup.cpp b/Engine/Plugins/Experimental/MediaViewer/Source/MediaViewer/Private/Widgets/SMediaViewerLibraryGroup.cpp
--- a/Engine/Plugins/Experimental/MediaViewer/Source/MediaViewer/Private/Widgets/SMediaViewerLibraryGroup.cpp
+++ b/Engine/Plugins/Experimental/MediaViewer/Source/MediaViewer/Private/Widgets/SMediaViewerLibraryGroup.cpp
@@ -105,6 +105,28 @@ void SMediaViewerLibraryGroup::Construct(const FArguments& InArgs, const TShared
 					SNew(SImage)
 					.Image(FAppStyle::Get().GetBrush("Icons.Delete"))
 					.DesiredSizeOverride(FVector2D(12))
+					.ToolTipText(LOCTEXT("DeleteGroupTooltip", "Remove this group."))
+				]
+			];
+	}
+	else if (Group->GetId() == InLibrary->GetHistoryGroupId())
+	{
+		const FButtonStyle& ButtonStyle = FMediaViewerStyle::Get().GetWidgetStyle<FButtonStyle>("LibraryButtons");
+
+		Inner->AddSlot()
+			.AutoWidth()
+			.Padding(0.f, 4.f, 5.f, 4.f)
+			.VAlign(VAlign_Center)
+			[
+				SNew(SButton)
+				.ContentPadding(FMargin(1.f, 3.f))
+				.OnClicked(this, &SMediaViewerLibraryGroup::OnRemoveButtonClicked)
+				.ButtonStyle(&ButtonStyle)
+				[
+					SNew(SImage)
+					.Image(FAppStyle::Get().GetBrush("Icons.Delete"))
+					.DesiredSizeOverride(FVector2D(12))
+					.ToolTipText(LOCTEXT("EmptyHistoryGroup", "Remove all items from the history."))
 				]
 			];
 	}
@@ -228,7 +250,11 @@ FReply SMediaViewerLibraryGroup::OnRemoveButtonClicked()
 		return FReply::Handled();
 	}
 
-	if (Library->CanRemoveGroup(Group->GetId()))
+	if (Group->GetId() == Library->GetHistoryGroupId())
+	{
+		Library->EmptyGroup(Group->GetId());
+	}
+	else if (Library->CanRemoveGroup(Group->GetId()))
 	{
 		Library->RemoveGroup(Group->GetId());
 	}
diff --git a/Engine/Plugins/Experimental/MediaViewer/Source/MediaViewer/Public/Library/IMediaViewerLibrary.h b/Engine/Plugins/Experimental/MediaViewer/Source/MediaViewer/Public/Library/IMediaViewerLibrary.h
--- a/Engine/Plugins/Experimental/MediaViewer/Source/MediaViewer/Public/Library/IMediaViewerLibrary.h
+++ b/Engine/Plugins/Experimental/MediaViewer/Source/MediaViewer/Public/Library/IMediaViewerLibrary.h
@@ -97,6 +97,12 @@ public:
 	 */
 	virtual TSharedPtr<FMediaViewerLibraryGroup> GetGroup(const FGuid& InGroupId) const = 0;
 
+	/**
+	 * Removes all items from a group. 
+	 * @returns The number of items removed.
+	 */
+	virtual int32 EmptyGroup(const FGuid& InGroupId) = 0;
+
 	/**
 	 * Scans the group and removes invalid entries.
 	 */
diff --git a/Engine/Plugins/Experimental/MediaViewer/Source/MediaViewer/Public/Library/MediaViewerLibraryGroup.h b/Engine/Plugins/Experimental/MediaViewer/Source/MediaViewer/Public/Library/MediaViewerLibraryGroup.h
--- a/Engine/Plugins/Experimental/MediaViewer/Source/MediaViewer/Public/Library/MediaViewerLibraryGroup.h
+++ b/Engine/Plugins/Experimental/MediaViewer/Source/MediaViewer/Public/Library/MediaViewerLibraryGroup.h
@@ -47,6 +47,8 @@ public:
 	
 	MEDIAVIEWER_API bool RemoveItemAt(int32 InIndex);
 
+	MEDIAVIEWER_API void Empty();
+
 	/** 
 	 * If true the group's contents are generated dynamically and not saved.
 	 * Its contents can be saved to the history group or copied to other groups.
