Media Viewer: Small tweak to new method.

diff --git a/Engine/Plugins/Experimental/MediaViewer/Source/MediaViewer/Private/Library/MediaViewerLibrary.cpp b/Engine/Plugins/Experimental/MediaViewer/Source/MediaViewer/Private/Library/MediaViewerLibrary.cpp
--- a/Engine/Plugins/Experimental/MediaViewer/Source/MediaViewer/Private/Library/MediaViewerLibrary.cpp
+++ b/Engine/Plugins/Experimental/MediaViewer/Source/MediaViewer/Private/Library/MediaViewerLibrary.cpp
@@ -134,9 +134,7 @@ int32 FMediaViewerLibrary::EmptyGroup(const FGuid& InGroupId)
 		return 0;
 	}
 
-	const int32 Count = Group->GetItems().Num();
-
-	Group->Empty();
+	const int32 Count = Group->Empty();	
 
 	OnChanged(IMediaViewerLibrary::EChangeType::ItemRemoved);
 
diff --git a/Engine/Plugins/Experimental/MediaViewer/Source/MediaViewer/Private/Library/MediaViewerLibraryGroup.cpp b/Engine/Plugins/Experimental/MediaViewer/Source/MediaViewer/Private/Library/MediaViewerLibraryGroup.cpp
--- a/Engine/Plugins/Experimental/MediaViewer/Source/MediaViewer/Private/Library/MediaViewerLibraryGroup.cpp
+++ b/Engine/Plugins/Experimental/MediaViewer/Source/MediaViewer/Private/Library/MediaViewerLibraryGroup.cpp
@@ -95,9 +95,16 @@ bool FMediaViewerLibraryGroup::RemoveItemAt(int32 InIndex)
 	return true;
 }
 
-void FMediaViewerLibraryGroup::Empty()
+int32 FMediaViewerLibraryGroup::Empty()
 {
-	Items.Empty();
+	if (!bDynamic)
+	{
+		const int32 Count = Items.Num();
+		Items.Empty();
+		return Count;
+	}
+
+	return 0;
 }
 
 bool FMediaViewerLibraryGroup::IsDynamic() const
diff --git a/Engine/Plugins/Experimental/MediaViewer/Source/MediaViewer/Public/Library/MediaViewerLibraryGroup.h b/Engine/Plugins/Experimental/MediaViewer/Source/MediaViewer/Public/Library/MediaViewerLibraryGroup.h
--- a/Engine/Plugins/Experimental/MediaViewer/Source/MediaViewer/Public/Library/MediaViewerLibraryGroup.h
+++ b/Engine/Plugins/Experimental/MediaViewer/Source/MediaViewer/Public/Library/MediaViewerLibraryGroup.h
@@ -47,7 +47,8 @@ public:
 	
 	MEDIAVIEWER_API bool RemoveItemAt(int32 InIndex);
 
-	MEDIAVIEWER_API void Empty();
+	/** Returns the number of elements removed. */
+	MEDIAVIEWER_API int32 Empty();
 
 	/** 
 	 * If true the group's contents are generated dynamically and not saved.
