Media Viewer: Fixed external drag interface.

diff --git a/Engine/Plugins/Experimental/MediaViewer/Source/MediaViewer/Private/Widgets/SMediaImageViewerOverlay.cpp b/Engine/Plugins/Experimental/MediaViewer/Source/MediaViewer/Private/Widgets/SMediaImageViewerOverlay.cpp
--- a/Engine/Plugins/Experimental/MediaViewer/Source/MediaViewer/Private/Widgets/SMediaImageViewerOverlay.cpp
+++ b/Engine/Plugins/Experimental/MediaViewer/Source/MediaViewer/Private/Widgets/SMediaImageViewerOverlay.cpp
@@ -358,7 +358,11 @@ void SMediaImageViewerOverlay::OnDragButtonUp(FName InKeyName)
 			break;
 
 		case EDragType::External:
-			UpdateDragExternal();
+			if (TSharedPtr<FMediaImageViewer> ImageViewer = Delegates->GetImageViewer.Execute(Position))
+			{
+				const FVector2D CursorLocation = Delegates->GetCursorLocation.Execute();
+				ImageViewer->OnTrackingStopped(FIntPoint(FMath::RoundToInt(CursorLocation.X), FMath::RoundToInt(CursorLocation.Y)));
+			}
 			break;
 	}
 
@@ -408,7 +412,7 @@ void SMediaImageViewerOverlay::UpdateDragExternal()
 	if (TSharedPtr<FMediaImageViewer> ImageViewer = Delegates->GetImageViewer.Execute(Position))
 	{
 		const FVector2D CursorLocation = Delegates->GetCursorLocation.Execute();
-		ImageViewer->OnTrackingStopped(FIntPoint(FMath::RoundToInt(CursorLocation.X), FMath::RoundToInt(CursorLocation.Y)));
+		ImageViewer->OnMouseMove(CursorLocation);
 	}
 }
 
