Media Viewer: Added a closed tab telemetry trigger.

diff --git a/Engine/Plugins/Experimental/MediaViewer/Source/MediaViewer/Private/Widgets/SMediaViewerTab.cpp b/Engine/Plugins/Experimental/MediaViewer/Source/MediaViewer/Private/Widgets/SMediaViewerTab.cpp
--- a/Engine/Plugins/Experimental/MediaViewer/Source/MediaViewer/Private/Widgets/SMediaViewerTab.cpp
+++ b/Engine/Plugins/Experimental/MediaViewer/Source/MediaViewer/Private/Widgets/SMediaViewerTab.cpp
@@ -22,6 +22,14 @@ void SMediaViewerTab::PrivateRegisterAttributes(FSlateAttributeDescriptor::FInit
 {
 }
 
+SMediaViewerTab::~SMediaViewerTab()
+{
+	if (FEngineAnalytics::IsAvailable())
+	{
+		FEngineAnalytics::GetProvider().RecordEvent(TEXT("Editor.Usage.MediaViewer"), TEXT("Action"), TEXT("ClosedTab"));
+	}
+}
+
 void SMediaViewerTab::Construct(const FArguments& InArgs, const FMediaViewerArgs& InMediaViewerArgs)
 {
 	TSharedPtr<FMediaImageViewer> ImageViewerLeft = InArgs._ImageViewerLeft;
diff --git a/Engine/Plugins/Experimental/MediaViewer/Source/MediaViewer/Public/Widgets/SMediaViewerTab.h b/Engine/Plugins/Experimental/MediaViewer/Source/MediaViewer/Public/Widgets/SMediaViewerTab.h
--- a/Engine/Plugins/Experimental/MediaViewer/Source/MediaViewer/Public/Widgets/SMediaViewerTab.h
+++ b/Engine/Plugins/Experimental/MediaViewer/Source/MediaViewer/Public/Widgets/SMediaViewerTab.h
@@ -32,7 +32,7 @@ class SMediaViewerTab : public SCompoundWidget
 public:
 	MEDIAVIEWER_API SMediaViewerTab();
 
-	virtual ~SMediaViewerTab() override = default;
+	virtual ~SMediaViewerTab() override;
 
 	MEDIAVIEWER_API void Construct(const FArguments& InArgs, const FMediaViewerArgs& InMediaViewerArgs);
 
