Media Viewer / Media Player Editor:
- Added scrub event to IMediaPlayerSlider.
- Integrated scrub event to allow dual scrubbing.
- Fixed bug in file drag/drop for media viewer not allowing comparisons to current image.

diff --git a/Engine/Plugins/Experimental/MediaViewer/Source/MediaViewer/Private/Widgets/MediaViewerSettings.h b/Engine/Plugins/Experimental/MediaViewer/Source/MediaViewer/Private/Widgets/MediaViewerSettings.h
--- a/Engine/Plugins/Experimental/MediaViewer/Source/MediaViewer/Private/Widgets/MediaViewerSettings.h
+++ b/Engine/Plugins/Experimental/MediaViewer/Source/MediaViewer/Private/Widgets/MediaViewerSettings.h
@@ -13,6 +13,22 @@
 
 class UTexture2D;
 
+UENUM()
+enum class EMediaViewerMediaSyncType : uint8
+{
+	/** The offset is relative to the start time on both videos. A scrub change of 2s in video A will move video B also by 2s. */
+	RelativeTime,
+
+	/** The offset is relative to slider value on both videos. A scrub change of 10% will in video A will move video B by 10%. */
+	RelativeProgress,
+
+	/** Time is synced exactly between videos. Video A changing from 5s to 10s will change video B from *anything* to 10s. */
+	AbsoluteTime,
+
+	/** Video progress is synced exactly between videos. Video A changing from 10% to 30% will change video B from *anything* to 30%. */
+	AbsoluteProgress
+};
+
 /** Settings for the viewer. */
 USTRUCT()
 struct FMediaViewerSettings
@@ -49,11 +65,15 @@ struct FMediaViewerSettings
 		meta = (UIMin = 0, ClampMin = 0, UIMax = 100, ClampMax = 100, Units = "Percent"))
 	float ABSplitterLocation = 50.f;
 
-	/** Opacity of the second/B image (0-100).
+	/** 
+	 * Opacity of the second/B image (0-100).
 	 * At 100%, the first/A image is rendered only on the left/top side.
 	 * Below 100%, the first/A image is rendered on the entire panel with B overlaid on top of it, translucently.
 	 */
 	UPROPERTY(EditAnywhere, Config, Category = "AB View", DisplayName = "B Image Opacity",
 		meta = (UIMin = 0, ClampMin = 0, UIMax = 100, ClampMax = 100, Units = "Percent"))
 	float SecondImageOpacity = 100.f;
+
+	UPROPERTY(EditAnywhere, Config, Category = "Media")
+	EMediaViewerMediaSyncType MediaSyncType = EMediaViewerMediaSyncType::RelativeTime;
 };
diff --git a/Engine/Plugins/Experimental/MediaViewer/Source/MediaViewer/Private/Widgets/SMediaSourceOverlay.cpp b/Engine/Plugins/Experimental/MediaViewer/Source/MediaViewer/Private/Widgets/SMediaSourceOverlay.cpp
--- a/Engine/Plugins/Experimental/MediaViewer/Source/MediaViewer/Private/Widgets/SMediaSourceOverlay.cpp
+++ b/Engine/Plugins/Experimental/MediaViewer/Source/MediaViewer/Private/Widgets/SMediaSourceOverlay.cpp
@@ -6,7 +6,6 @@
 #include "ImageViewers/MediaSourceImageViewer.h"
 #include "IMediaStreamPlayer.h"
 #include "MediaPlayer.h"
-#include "MediaPlayerEditorModule.h"
 #include "MediaStream.h"
 #include "MediaViewerStyle.h"
 #include "Misc/FrameRate.h"
@@ -43,6 +42,13 @@ namespace MediaSourceOverlay::Events
 	const FLazyName Pause       = "MediaSource_Pause";
 	const FLazyName StepForward = "MediaSource_StepForward";
 	const FLazyName Forward     = "MediaSource_Forward";
+	const FLazyName Scrub       = "MediaSource_Scrub";
+
+	struct FMediaImageViewerScrubEventParams : FMediaImageViewerEventParams
+	{
+		IMediaPlayerSlider::EScrubEventType ScrubEvent;
+		TVariant<FTimespan, float> ScrubData = {};
+	};
 }
 
 void SMediaSourceOverlay::PrivateRegisterAttributes(FSlateAttributeDescriptor::FInitializer&)
@@ -185,6 +191,7 @@ TSharedRef<SWidget> SMediaSourceOverlay::CreateSlider()
 
 	MediaPlayerSlider->SetSliderHandleColor(FSlateColor(EStyleColor::AccentBlue));
 	MediaPlayerSlider->SetVisibleWhenInactive(EVisibility::Visible);
+	MediaPlayerSlider->GetScrubEvent().AddSP(this, &SMediaSourceOverlay::OnScrub);
 
 	return MediaPlayerSlider;
 }
@@ -784,6 +791,258 @@ void SMediaSourceOverlay::ReceiveEvent(const FMediaImageViewerEventParams& InEve
 	{
 		Command_Forward(MediaPlayer);
 	}
+	else if (InEventParams.EventName == Events::Scrub)
+	{
+		ReceiveScrub(static_cast<const MediaSourceOverlay::Events::FMediaImageViewerScrubEventParams&>(InEventParams));
+	}
+}
+
+void SMediaSourceOverlay::OnScrub(IMediaPlayerSlider::EScrubEventType InSliderEvent, TConstArrayView<UMediaPlayer*> InMediaPlayers, float InScrubPosition)
+{
+	if (InMediaPlayers.Num() != 1 || !IsValid(InMediaPlayers[0]))
+	{
+		return;
+	}
+
+	UMediaPlayer* MediaPlayer = InMediaPlayers[0];
+
+	if (!Delegates.IsValid())
+	{
+		return;
+	}
+
+	if (!Delegates->AreTransformsLocked.IsBound() || !Delegates->GetSettings.IsBound())
+	{
+		return;
+	}
+
+	if (!Delegates->AreTransformsLocked.Execute())
+	{
+		return;
+	}
+
+	using namespace MediaSourceOverlay::Events;
+
+	const EMediaViewerMediaSyncType MediaSyncType = Delegates->GetSettings.Execute()->MediaSyncType;
+
+	FMediaImageViewerScrubEventParams Params;
+	Params.FromViewer = ImageViewerWeak.Pin();
+	Params.EventName = Scrub;
+	Params.ScrubEvent = InSliderEvent;
+
+	auto SetScrubData = [this, MediaPlayer, MediaSyncType, InScrubPosition, &Params]()
+		{
+			switch (MediaSyncType)
+			{
+				case EMediaViewerMediaSyncType::RelativeTime:
+					if (!SyncedScrubStartTime.IsType<FTimespan>())
+					{
+						return false;
+					}
+
+					Params.ScrubData.Set<FTimespan>((MediaPlayer->GetDuration() * InScrubPosition) - SyncedScrubStartTime.Get<FTimespan>());
+					break;
+
+				case EMediaViewerMediaSyncType::AbsoluteTime:
+					Params.ScrubData.Set<FTimespan>(MediaPlayer->GetDuration() * InScrubPosition);
+					break;
+
+				case EMediaViewerMediaSyncType::RelativeProgress:
+					if (!SyncedScrubStartTime.IsType<float>())
+					{
+						return false;
+					}
+
+					Params.ScrubData.Set<float>(InScrubPosition - SyncedScrubStartTime.Get<float>());
+					break;
+
+				case EMediaViewerMediaSyncType::AbsoluteProgress:
+					Params.ScrubData.Set<float>(InScrubPosition);
+					break;
+
+				default:
+					return false;
+			}			
+
+			return true;
+		};
+
+	switch (InSliderEvent)
+	{
+		case IMediaPlayerSlider::EScrubEventType::Begin:
+			switch (MediaSyncType)
+			{
+				case EMediaViewerMediaSyncType::RelativeTime:
+					SyncedScrubStartTime.Set<FTimespan>(MediaPlayer->GetDuration() * InScrubPosition);
+					break;
+
+				case EMediaViewerMediaSyncType::RelativeProgress:
+					SyncedScrubStartTime.Set<float>(InScrubPosition);
+					break;
+
+				default:
+					// Nothing
+					break;
+			};			
+
+			Delegates->ImageViewerEvent.Broadcast(Params);
+			break;
+
+		case IMediaPlayerSlider::EScrubEventType::Update:
+			if (SetScrubData())
+			{
+				Delegates->ImageViewerEvent.Broadcast(Params);
+			}
+			break;
+
+		case IMediaPlayerSlider::EScrubEventType::End:
+			if (SetScrubData())
+			{
+				Delegates->ImageViewerEvent.Broadcast(Params);
+			}
+
+			SyncedScrubStartTime = {};
+			break;
+
+		default:
+			// Do nothing
+			break;
+	}
+}
+
+void SMediaSourceOverlay::ReceiveScrub(const MediaSourceOverlay::Events::FMediaImageViewerScrubEventParams& InEventParams)
+{
+	TSharedPtr<FMediaSourceImageViewer> ImageViewer = ImageViewerWeak.Pin();
+
+	if (!ImageViewer.IsValid())
+	{
+		return;
+	}
+
+	using namespace MediaSourceOverlay::Events;
+
+	if (InEventParams.FromViewer.Get() == ImageViewer.Get())
+	{
+		return;
+	}
+
+	UMediaPlayer* MediaPlayer = GetMediaPlayer();
+
+	if (!MediaPlayer)
+	{
+		return;
+	}
+
+	if (!Delegates.IsValid())
+	{
+		return;
+	}
+
+	if (!Delegates->GetSettings.IsBound())
+	{
+		return;
+	}
+
+	if (!MediaPlayer->SupportsScrubbing())
+	{
+		return;
+	}
+
+	const EMediaViewerMediaSyncType MediaSyncType = Delegates->GetSettings.Execute()->MediaSyncType;
+
+	auto UpdateScrub = [this, MediaPlayer, MediaSyncType, &InEventParams]()
+		{
+			switch (MediaSyncType)
+			{
+				case EMediaViewerMediaSyncType::RelativeTime:
+					if (SyncedScrubStartTime.IsType<FTimespan>() && InEventParams.ScrubData.IsType<FTimespan>())
+					{
+						MediaPlayer->Seek(FTimespan::FromSeconds(FMath::Clamp(
+							(SyncedScrubStartTime.Get<FTimespan>() + InEventParams.ScrubData.Get<FTimespan>()).GetTotalSeconds(),
+							0.0,
+							MediaPlayer->GetDuration().GetTotalSeconds()
+						)));
+					}
+					break;
+
+				case EMediaViewerMediaSyncType::AbsoluteTime:
+					if (InEventParams.ScrubData.IsType<FTimespan>())
+					{
+						MediaPlayer->Seek(FTimespan::FromSeconds(FMath::Clamp(
+							InEventParams.ScrubData.Get<FTimespan>().GetTotalSeconds(),
+							0.0,
+							MediaPlayer->GetDuration().GetTotalSeconds()
+						)));
+					}
+					break;
+
+				case EMediaViewerMediaSyncType::RelativeProgress:
+					if (SyncedScrubStartTime.IsType<float>() && InEventParams.ScrubData.IsType<float>())
+					{
+						MediaPlayer->Seek(MediaPlayer->GetDuration() * FMath::Clamp(
+							SyncedScrubStartTime.Get<float>() + InEventParams.ScrubData.Get<float>(),
+							0.0,
+							1.0
+						));
+					}
+					break;
+
+				case EMediaViewerMediaSyncType::AbsoluteProgress:
+					if (InEventParams.ScrubData.IsType<float>())
+					{
+						MediaPlayer->Seek(MediaPlayer->GetDuration() * FMath::Clamp(
+							InEventParams.ScrubData.Get<float>(),
+							0.0,
+							1.0
+						));
+					}
+					break;
+
+				default:
+					// Do nothing
+					break;
+			}
+		};
+
+	switch (InEventParams.ScrubEvent)
+	{
+		case IMediaPlayerSlider::EScrubEventType::Begin:
+		{
+			switch (MediaSyncType)
+			{
+				case EMediaViewerMediaSyncType::RelativeTime:
+					SyncedScrubStartTime.Set<FTimespan>(MediaPlayer->GetTime());
+					break;
+
+				case EMediaViewerMediaSyncType::RelativeProgress:
+					SyncedScrubStartTime.Set<float>(MediaPlayer->GetTime().GetTotalSeconds() / MediaPlayer->GetDuration().GetTotalSeconds());
+					break;
+
+				default:
+					// Nothing
+					break;
+			};
+
+			break;
+		}
+
+		case IMediaPlayerSlider::EScrubEventType::Update:
+		{
+			UpdateScrub();
+			break;
+		}
+
+		case IMediaPlayerSlider::EScrubEventType::End:
+		{
+			UpdateScrub();
+			SyncedScrubStartTime = {};
+			break;
+		}
+
+		default:
+			// Do nothing
+			break;
+	}
 }
 
 bool SMediaSourceOverlay::Command_Rewind(UMediaPlayer* InMediaPlayer)
diff --git a/Engine/Plugins/Experimental/MediaViewer/Source/MediaViewer/Private/Widgets/SMediaSourceOverlay.h b/Engine/Plugins/Experimental/MediaViewer/Source/MediaViewer/Private/Widgets/SMediaSourceOverlay.h
--- a/Engine/Plugins/Experimental/MediaViewer/Source/MediaViewer/Private/Widgets/SMediaSourceOverlay.h
+++ b/Engine/Plugins/Experimental/MediaViewer/Source/MediaViewer/Private/Widgets/SMediaSourceOverlay.h
@@ -5,7 +5,9 @@
 #include "Widgets/SCompoundWidget.h"
 
 #include "Internationalization/Text.h"
-#include "Misc/Optional.h"
+#include "MediaPlayerEditorModule.h"
+#include "Misc/Timespan.h"
+#include "Misc/TVariant.h"
 
 class IMediaStreamPlayer;
 class SBox;
@@ -24,6 +26,11 @@ namespace UE::MediaViewer::Private
 class FMediaSourceImageViewer;
 struct FMediaImageViewerEventParams;
 struct FMediaViewerDelegates;
+
+namespace MediaSourceOverlay::Events
+{
+	struct FMediaImageViewerScrubEventParams;
+}
 	
 class SMediaSourceOverlay : public SCompoundWidget
 {
@@ -50,6 +57,7 @@ protected:
 	TOptional<float> FrameRateFloat;
 	mutable TOptional<FText> TotalFrames;
 	mutable TOptional<FText> TotalTime;
+	TVariant<FTimespan, float> SyncedScrubStartTime;
 
 	/** The last time the mouse was moved over the widget. */
 	double LastInteractionTime = -1.0;
@@ -98,6 +106,8 @@ protected:
 
 	void ReceiveEvent(const FMediaImageViewerEventParams& InEventParams);
 
+	void OnScrub(IMediaPlayerSlider::EScrubEventType InScrubEvent, TConstArrayView<UMediaPlayer*> InMediaPlayers, float InScrubPosition);
+
 	static bool Command_Rewind(UMediaPlayer* InMediaPlayer);
 	static bool Command_Reverse(UMediaPlayer* InMediaPlayer);
 	static bool Command_StepBack(UMediaPlayer* InMediaPlayer);
@@ -105,6 +115,8 @@ protected:
 	static bool Command_Pause(UMediaPlayer* InMediaPlayer);
 	static bool Command_StepForward(UMediaPlayer* InMediaPlayer);
 	static bool Command_Forward(UMediaPlayer* InMediaPlayer);
+	
+	void ReceiveScrub(const MediaSourceOverlay::Events::FMediaImageViewerScrubEventParams& InEventParams);
 };
 
 } // UE::MediaViewer::Private
\ No newline at end of file
diff --git a/Engine/Plugins/Experimental/MediaViewer/Source/MediaViewer/Private/Widgets/SMediaViewerDropTarget.cpp b/Engine/Plugins/Experimental/MediaViewer/Source/MediaViewer/Private/Widgets/SMediaViewerDropTarget.cpp
--- a/Engine/Plugins/Experimental/MediaViewer/Source/MediaViewer/Private/Widgets/SMediaViewerDropTarget.cpp
+++ b/Engine/Plugins/Experimental/MediaViewer/Source/MediaViewer/Private/Widgets/SMediaViewerDropTarget.cpp
@@ -273,7 +273,7 @@ void SMediaViewerDropTarget::HandleDroppeFileOp(const FExternalDragOperation& In
 		}
 
 		Delegates->SetImageViewer.Execute(
-			EMediaImageViewerPosition::First,
+			Position,
 			MakeShared<FMediaSourceImageViewer>(Sources[0], FText::FromString(Sources[0]->GetFilePath()))
 		);
 	}
diff --git a/Engine/Plugins/Media/MediaPlayerEditor/Source/MediaPlayerEditor/Private/Widgets/SMediaPlayerSlider.cpp b/Engine/Plugins/Media/MediaPlayerEditor/Source/MediaPlayerEditor/Private/Widgets/SMediaPlayerSlider.cpp
--- a/Engine/Plugins/Media/MediaPlayerEditor/Source/MediaPlayerEditor/Private/Widgets/SMediaPlayerSlider.cpp
+++ b/Engine/Plugins/Media/MediaPlayerEditor/Source/MediaPlayerEditor/Private/Widgets/SMediaPlayerSlider.cpp
@@ -52,10 +52,15 @@ bool SMediaPlayerSlider::DoesMediaPlayerSupportSeeking() const
 
 void SMediaPlayerSlider::OnScrubBegin()
 {
+	TArray<UMediaPlayer*, TInlineAllocator<1>> MediaPlayers;
+	MediaPlayers.Reserve(MediaPlayerEntries.Num());
+
 	for (FMediaPlayerEntry& Entry : MediaPlayerEntries)
 	{
 		if (UMediaPlayer* MediaPlayer = Entry.MediaPlayerWeak.Get())
 		{
+			MediaPlayers.Add(MediaPlayer);
+
 			Entry.ScrubValue = static_cast<float>(FTimespan::Ratio(MediaPlayer->GetDisplayTime(), MediaPlayer->GetDuration()));
 			Entry.LastScrubValue = Entry.ScrubValue;
 
@@ -66,14 +71,21 @@ void SMediaPlayerSlider::OnScrubBegin()
 			}
 		}
 	}
+
+	ScrubEvent.Broadcast(IMediaPlayerSlider::EScrubEventType::Begin, MediaPlayers, GetSliderValue());
 }
 
 void SMediaPlayerSlider::OnScrubEnd()
 {
+	TArray<UMediaPlayer*, TInlineAllocator<1>> MediaPlayers;
+	MediaPlayers.Reserve(MediaPlayerEntries.Num());
+
 	for (FMediaPlayerEntry& Entry : MediaPlayerEntries)
 	{
 		if (UMediaPlayer* MediaPlayer = Entry.MediaPlayerWeak.Get())
 		{
+			MediaPlayers.Add(MediaPlayer);
+
 			// Set playback position to scrub value when drag ends
 			if (Entry.LastScrubValue != Entry.ScrubValue && MediaPlayer->SupportsSeeking())
 			{
@@ -86,14 +98,21 @@ void SMediaPlayerSlider::OnScrubEnd()
 			}
 		}
 	}
+
+	ScrubEvent.Broadcast(IMediaPlayerSlider::EScrubEventType::End, MediaPlayers, GetSliderValue());
 }
 
 void SMediaPlayerSlider::Seek(float InPlaybackPosition)
 {
+	TArray<UMediaPlayer*, TInlineAllocator<1>> MediaPlayers;
+	MediaPlayers.Reserve(MediaPlayerEntries.Num());
+
 	for (FMediaPlayerEntry& Entry : MediaPlayerEntries)
 	{
 		if (UMediaPlayer* MediaPlayer = Entry.MediaPlayerWeak.Get())
 		{
+			MediaPlayers.Add(MediaPlayer);
+
 			Entry.ScrubValue = InPlaybackPosition;
 
 			if (!ScrubberSlider->HasMouseCapture() || MediaPlayer->SupportsScrubbing())
@@ -103,6 +122,8 @@ void SMediaPlayerSlider::Seek(float InPlaybackPosition)
 			}
 		}
 	}
+
+	ScrubEvent.Broadcast(IMediaPlayerSlider::EScrubEventType::Update, MediaPlayers, InPlaybackPosition);
 }
 
 float SMediaPlayerSlider::GetPlaybackPosition() const
@@ -150,6 +171,16 @@ EVisibility SMediaPlayerSlider::GetScrubberVisibility() const
 	return bIsActive ? EVisibility::Visible : VisibilityWhenInactive;
 }
 
+float SMediaPlayerSlider::GetSliderValue() const
+{
+	if (ScrubberSlider.IsValid())
+	{
+		return ScrubberSlider->GetValue();
+	}
+
+	return 0.f;
+}
+
 void SMediaPlayerSlider::SetSliderHandleColor(const FSlateColor& InSliderColor)
 {
 	if (ScrubberSlider)
@@ -171,6 +202,11 @@ void SMediaPlayerSlider::SetVisibleWhenInactive(EVisibility InVisibility)
 	VisibilityWhenInactive = InVisibility;
 }
 
+IMediaPlayerSlider::FScrubEvent::RegistrationType& SMediaPlayerSlider::GetScrubEvent()
+{
+	return ScrubEvent;
+}
+
 /** Find a valid player entry for the given track type. */
 const SMediaPlayerSlider::FMediaPlayerEntry* SMediaPlayerSlider::FindValidPlayerEntryForTrackType(EMediaPlayerTrack InTrackType) const
 {
diff --git a/Engine/Plugins/Media/MediaPlayerEditor/Source/MediaPlayerEditor/Private/Widgets/SMediaPlayerSlider.h b/Engine/Plugins/Media/MediaPlayerEditor/Source/MediaPlayerEditor/Private/Widgets/SMediaPlayerSlider.h
--- a/Engine/Plugins/Media/MediaPlayerEditor/Source/MediaPlayerEditor/Private/Widgets/SMediaPlayerSlider.h
+++ b/Engine/Plugins/Media/MediaPlayerEditor/Source/MediaPlayerEditor/Private/Widgets/SMediaPlayerSlider.h
@@ -35,6 +35,7 @@ public:
 	virtual void SetSliderHandleColor(const FSlateColor& InSliderColor) override;
 	virtual void SetSliderBarColor(const FSlateColor& InSliderColor) override;
 	virtual void SetVisibleWhenInactive(EVisibility InVisibility) override;
+	virtual FScrubEvent::RegistrationType& GetScrubEvent() override;
 	//~End IMediaPlayerSlider
 
 private:
@@ -44,6 +45,7 @@ private:
 	void Seek(float InPlaybackPosition);
 	float GetPlaybackPosition() const;
 	EVisibility GetScrubberVisibility() const;
+	float GetSliderValue() const;
 
 	/** The scrubber visibility when inactive. */
 	EVisibility VisibilityWhenInactive = EVisibility::Hidden;
@@ -70,4 +72,7 @@ private:
 	
 	/** Holds the scrubber slider. */
 	TSharedPtr<SSlider> ScrubberSlider;
+
+	/** Event subscribers can use to detect scrubbing. */
+	IMediaPlayerSlider::FScrubEvent ScrubEvent;
 };
diff --git a/Engine/Plugins/Media/MediaPlayerEditor/Source/MediaPlayerEditor/Public/MediaPlayerEditorModule.h b/Engine/Plugins/Media/MediaPlayerEditor/Source/MediaPlayerEditor/Public/MediaPlayerEditorModule.h
--- a/Engine/Plugins/Media/MediaPlayerEditor/Source/MediaPlayerEditor/Public/MediaPlayerEditorModule.h
+++ b/Engine/Plugins/Media/MediaPlayerEditor/Source/MediaPlayerEditor/Public/MediaPlayerEditorModule.h
@@ -25,6 +25,22 @@ public:
 
 	/** Set the Scrubber Slider Visibility, when the player is inactive */
 	virtual void SetVisibleWhenInactive(EVisibility InVisibility) = 0;
+
+	/** When scrubbing of the created slider occurs, this enum is used in FScrubEvent to inform subscribers about the slider state. */
+	enum class EScrubEventType : uint8
+	{
+		/** Scrubbing has started. */
+		Begin,
+		/** Scrubbing is ongoing and the position of the scrub has changed. */
+		Update,
+		/** Scrubbing has stopped. */
+		End
+	};
+
+	DECLARE_MULTICAST_DELEGATE_ThreeParams(FScrubEvent, EScrubEventType /* Event type */, TConstArrayView<UMediaPlayer*> /* Media Players */, float /* Slider/Scrub Value */)
+
+	/** Event subscribers can use to detect scrubbing. */
+	virtual FScrubEvent::RegistrationType& GetScrubEvent() = 0;
 };
 
 /**
