Media Viewer: Fixed potential /0 error.

diff --git a/Engine/Plugins/Experimental/MediaViewer/Source/MediaViewer/Private/Widgets/SMediaSourceOverlay.cpp b/Engine/Plugins/Experimental/MediaViewer/Source/MediaViewer/Private/Widgets/SMediaSourceOverlay.cpp
--- a/Engine/Plugins/Experimental/MediaViewer/Source/MediaViewer/Private/Widgets/SMediaSourceOverlay.cpp
+++ b/Engine/Plugins/Experimental/MediaViewer/Source/MediaViewer/Private/Widgets/SMediaSourceOverlay.cpp
@@ -1015,8 +1015,16 @@ void SMediaSourceOverlay::ReceiveScrub(const MediaSourceOverlay::Events::FMediaI
 					break;
 
 				case EMediaViewerMediaSyncType::RelativeProgress:
-					SyncedScrubStartTime.Set<float>(MediaPlayer->GetTime().GetTotalSeconds() / MediaPlayer->GetDuration().GetTotalSeconds());
+				{
+					const float MediaDuration = MediaPlayer->GetDuration().GetTotalSeconds();
+
+					if (MediaDuration > 0)
+					{
+						SyncedScrubStartTime.Set<float>(MediaPlayer->GetTime().GetTotalSeconds() / MediaDuration);
+					}
+
 					break;
+				}
 
 				default:
 					// Nothing
