Media Viewer: Extra keybinds for videos
- Fixed spam of updates when offset / scale was not changing.
- Exposed delegates to the public api to allow it to be used in the image viewer overlay.

diff --git a/Engine/Plugins/Experimental/MediaViewer/Source/MediaViewer/Private/ImageViewers/MediaSourceImageViewer.cpp b/Engine/Plugins/Experimental/MediaViewer/Source/MediaViewer/Private/ImageViewers/MediaSourceImageViewer.cpp
--- a/Engine/Plugins/Experimental/MediaViewer/Source/MediaViewer/Private/ImageViewers/MediaSourceImageViewer.cpp
+++ b/Engine/Plugins/Experimental/MediaViewer/Source/MediaViewer/Private/ImageViewers/MediaSourceImageViewer.cpp
@@ -442,14 +442,14 @@ void FMediaSourceImageViewer::AddReferencedObjects(FReferenceCollector& Collecto
 	}
 }
 
-TSharedPtr<SWidget> FMediaSourceImageViewer::GetOverlayWidget(EMediaImageViewerPosition InPosition, const TSharedPtr<SMediaViewerTab>& InViewerTab)
+TSharedPtr<SWidget> FMediaSourceImageViewer::GetOverlayWidget(EMediaImageViewerPosition InPosition, const TSharedPtr<FMediaViewerDelegates>& InDelegates)
 {
 	return SNew(SBox)
 		.HAlign(HAlign_Fill)
 		.VAlign(VAlign_Bottom)
 		.Padding(5.f)
 		[
-			SNew(SMediaSourceOverlay, SharedThis(this), InViewerTab)
+			SNew(SMediaSourceOverlay, SharedThis(this), InPosition, InDelegates)
 		];
 }
 
diff --git a/Engine/Plugins/Experimental/MediaViewer/Source/MediaViewer/Private/ImageViewers/MediaSourceImageViewer.h b/Engine/Plugins/Experimental/MediaViewer/Source/MediaViewer/Private/ImageViewers/MediaSourceImageViewer.h
--- a/Engine/Plugins/Experimental/MediaViewer/Source/MediaViewer/Private/ImageViewers/MediaSourceImageViewer.h
+++ b/Engine/Plugins/Experimental/MediaViewer/Source/MediaViewer/Private/ImageViewers/MediaSourceImageViewer.h
@@ -123,7 +123,7 @@ public:
 	virtual TSharedPtr<FMediaViewerLibraryItem> CreateLibraryItem() const override;
 	virtual TOptional<TVariant<FColor, FLinearColor>> GetPixelColor(const FIntPoint& InPixelCoords, int32 InMipLevel) const override;
 	virtual TSharedPtr<FStructOnScope> GetCustomSettingsOnScope() const override;
-	virtual TSharedPtr<SWidget> GetOverlayWidget(EMediaImageViewerPosition InPosition, const TSharedPtr<SMediaViewerTab>& InViewerTab) override;
+	virtual TSharedPtr<SWidget> GetOverlayWidget(EMediaImageViewerPosition InPosition, const TSharedPtr<FMediaViewerDelegates>& InDelegates) override;
 	virtual void ExtendStatusBar(UE::MediaViewer::FMediaImageStatusBarExtender& InOutStatusBarExtender) override;
 	//~ End FMediaImageViewer
 
diff --git a/Engine/Plugins/Experimental/MediaViewer/Source/MediaViewer/Private/MediaViewerCommands.cpp b/Engine/Plugins/Experimental/MediaViewer/Source/MediaViewer/Private/MediaViewerCommands.cpp
--- a/Engine/Plugins/Experimental/MediaViewer/Source/MediaViewer/Private/MediaViewerCommands.cpp
+++ b/Engine/Plugins/Experimental/MediaViewer/Source/MediaViewer/Private/MediaViewerCommands.cpp
@@ -222,7 +222,7 @@ void FMediaViewerCommands::RegisterCommands()
 		"Sync Transforms",
 		"Syncs the transforms of all images so that changes made to one are made to all of them.",
 		EUserInterfaceActionType::ToggleButton,
-		FInputChord(EKeys::L)
+		FInputChord(EKeys::T)
 	);
 
 	UI_COMMAND(
@@ -288,6 +288,65 @@ void FMediaViewerCommands::RegisterCommands()
 		EUserInterfaceActionType::Button,
 		FInputChord(EKeys::S, EModifierKey::Alt)
 	);
+
+	UI_COMMAND(
+		QuickRewindVideo,
+		"Quick Rewind Video",
+		"Jumps back 10s in the currently active video.",
+		EUserInterfaceActionType::Button,
+		FInputChord(EKeys::J, EModifierKey::Shift)
+	);
+
+	UI_COMMAND(
+		RewindVideo,
+		"Rewind Video",
+		"Jumps back 1s in the currently active video.",
+		EUserInterfaceActionType::Button,
+		FInputChord(EKeys::J)
+	);
+
+	UI_COMMAND(
+		StepBackVideo,
+		"Step Back Video",
+		"Jumps back 1 frame in the currently active video.",
+		EUserInterfaceActionType::Button,
+		FInputChord(EKeys::J, EModifierKey::Control),
+		FInputChord(EKeys::J, EModifierKey::Command)
+	);
+
+	UI_COMMAND(
+		ToggleVideoPlay,
+		"Toggle Play/Pause",
+		"Toggles between play and pause on the current video.",
+		EUserInterfaceActionType::Button,
+		FInputChord(EKeys::K),
+		FInputChord(EKeys::SpaceBar)
+	);
+
+	UI_COMMAND(
+		StepForwardVideo,
+		"Step Forward Video",
+		"Jumps forward 1 frame in the currently active video.",
+		EUserInterfaceActionType::Button,
+		FInputChord(EKeys::L, EModifierKey::Control),
+		FInputChord(EKeys::L, EModifierKey::Command)
+	);
+
+	UI_COMMAND(
+		FastForwardVideo,
+		"Fast Forward Video",
+		"Jumps forward 1s in the currently active video.",
+		EUserInterfaceActionType::Button,
+		FInputChord(EKeys::L)
+	);
+
+	UI_COMMAND(
+		QuickFastForwardVideo,
+		"Quick Fast Forward Video",
+		"Jumps forward 10s in the currently active video.",
+		EUserInterfaceActionType::Button,
+		FInputChord(EKeys::L, EModifierKey::Shift)
+	);
 }
 
 } // UE::MediaViewer::Private
diff --git a/Engine/Plugins/Experimental/MediaViewer/Source/MediaViewer/Private/MediaViewerCommands.h b/Engine/Plugins/Experimental/MediaViewer/Source/MediaViewer/Private/MediaViewerCommands.h
--- a/Engine/Plugins/Experimental/MediaViewer/Source/MediaViewer/Private/MediaViewerCommands.h
+++ b/Engine/Plugins/Experimental/MediaViewer/Source/MediaViewer/Private/MediaViewerCommands.h
@@ -60,6 +60,14 @@ public:
 	TSharedPtr<FUICommandInfo> SwapAB;
 
 	TSharedPtr<FUICommandInfo> Snapshot;
+
+	TSharedPtr<FUICommandInfo> QuickRewindVideo;
+	TSharedPtr<FUICommandInfo> RewindVideo;
+	TSharedPtr<FUICommandInfo> StepBackVideo;
+	TSharedPtr<FUICommandInfo> ToggleVideoPlay;
+	TSharedPtr<FUICommandInfo> StepForwardVideo;
+	TSharedPtr<FUICommandInfo> FastForwardVideo;
+	TSharedPtr<FUICommandInfo> QuickFastForwardVideo;
 };
 
 } // UE::MediaViewer::Private
diff --git a/Engine/Plugins/Experimental/MediaViewer/Source/MediaViewer/Private/MediaViewerModule.cpp b/Engine/Plugins/Experimental/MediaViewer/Source/MediaViewer/Private/MediaViewerModule.cpp
--- a/Engine/Plugins/Experimental/MediaViewer/Source/MediaViewer/Private/MediaViewerModule.cpp
+++ b/Engine/Plugins/Experimental/MediaViewer/Source/MediaViewer/Private/MediaViewerModule.cpp
@@ -17,11 +17,11 @@
 #include "MediaViewer.h"
 #include "MediaViewerCommands.h"
 #include "MediaViewerContentBrowserIntegration.h"
+#include "MediaViewerDelegates.h"
 #include "MediaViewerStyle.h"
 #include "Misc/NotNull.h"
 #include "Modules/ModuleManager.h"
 #include "Widgets/Docking/SDockTab.h"
-#include "Widgets/MediaViewerDelegates.h"
 #include "Widgets/SMediaViewer.h"
 #include "Widgets/SMediaViewerLibrary.h"
 #include "Widgets/SMediaViewerTab.h"
diff --git a/Engine/Plugins/Experimental/MediaViewer/Source/MediaViewer/Private/Widgets/MediaViewerDelegates.h b/Engine/Plugins/Experimental/MediaViewer/Source/MediaViewer/Private/Widgets/MediaViewerDelegates.h
deleted file mode 100644
diff --git a/Engine/Plugins/Experimental/MediaViewer/Source/MediaViewer/Private/Widgets/SMediaImageViewerDetails.cpp b/Engine/Plugins/Experimental/MediaViewer/Source/MediaViewer/Private/Widgets/SMediaImageViewerDetails.cpp
--- a/Engine/Plugins/Experimental/MediaViewer/Source/MediaViewer/Private/Widgets/SMediaImageViewerDetails.cpp
+++ b/Engine/Plugins/Experimental/MediaViewer/Source/MediaViewer/Private/Widgets/SMediaImageViewerDetails.cpp
@@ -4,10 +4,10 @@
 
 #include "ImageViewer/MediaImageViewer.h"
 #include "MediaViewer.h"
+#include "MediaViewerDelegates.h"
 #include "MediaViewerUtils.h"
 #include "UObject/StructOnScope.h"
 #include "Widgets/Layout/SScrollBox.h"
-#include "Widgets/MediaViewerDelegates.h"
 
 #define LOCTEXT_NAMESPACE "SMediaImageViewerDetails"
 
diff --git a/Engine/Plugins/Experimental/MediaViewer/Source/MediaViewer/Private/Widgets/SMediaImageViewerDetails.h b/Engine/Plugins/Experimental/MediaViewer/Source/MediaViewer/Private/Widgets/SMediaImageViewerDetails.h
--- a/Engine/Plugins/Experimental/MediaViewer/Source/MediaViewer/Private/Widgets/SMediaImageViewerDetails.h
+++ b/Engine/Plugins/Experimental/MediaViewer/Source/MediaViewer/Private/Widgets/SMediaImageViewerDetails.h
@@ -4,8 +4,6 @@
 
 #include "Widgets/SCompoundWidget.h"
 
-#include "Widgets/MediaViewerDelegates.h"
-
 class FStructOnScope;
 class IStructureDetailsView;
 
@@ -13,13 +11,12 @@ namespace UE::MediaViewer
 {
 	class FMediaImageViewer;
 	enum class EMediaImageViewerPosition : uint8;
+	struct FMediaViewerDelegates;
 }
 
 namespace UE::MediaViewer::Private
 {
 
-struct FMediaViewerDelegates;
-
 class SMediaImageViewerDetails : public SCompoundWidget
 {
 	SLATE_DECLARE_WIDGET(SMediaImageViewerDetails, SCompoundWidget)
diff --git a/Engine/Plugins/Experimental/MediaViewer/Source/MediaViewer/Private/Widgets/SMediaImageViewerOverlay.cpp b/Engine/Plugins/Experimental/MediaViewer/Source/MediaViewer/Private/Widgets/SMediaImageViewerOverlay.cpp
--- a/Engine/Plugins/Experimental/MediaViewer/Source/MediaViewer/Private/Widgets/SMediaImageViewerOverlay.cpp
+++ b/Engine/Plugins/Experimental/MediaViewer/Source/MediaViewer/Private/Widgets/SMediaImageViewerOverlay.cpp
@@ -17,6 +17,7 @@
 #include "Math/Vector2D.h"
 #include "MediaViewer.h"
 #include "MediaViewerCommands.h"
+#include "MediaViewerDelegates.h"
 #include "Widgets/Layout/SScissorRectBox.h"
 #include "Widgets/SBoxPanel.h"
 #include "Widgets/SMediaImageViewerStatusBar.h"
@@ -33,6 +34,7 @@ const FName DropTargetName = TEXT("DropTarget");
 
 SMediaImageViewerOverlay::SMediaImageViewerOverlay()
 	: CommandList(MakeShared<FUICommandList>())
+	, bComparisonView(false)
 	, DragButtonName(NAME_None)
 	, DragType(EDragType::None)
 	, bOverlayEnabled(true)
@@ -48,6 +50,7 @@ void SMediaImageViewerOverlay::Construct(const FArguments& InArgs, EMediaImageVi
 {
 	Position = InPosition;
 	Delegates = InDelegates;
+	bComparisonView = InArgs._bComparisonView;
 	bScaleToFit = InArgs._bScaleToFit;
 
 	TSharedPtr<FMediaImageViewer> ImageViewer = Delegates->GetImageViewer.Execute(Position);
@@ -58,11 +61,16 @@ void SMediaImageViewerOverlay::Construct(const FArguments& InArgs, EMediaImageVi
 	}
 
 	CachedItem = ImageViewer->CreateLibraryItem();
+}
 
-	BindCommands();
-	InDelegates->GetCommandList.Execute()->Append(CommandList);
+void SMediaImageViewerOverlay::PostConstruct()
+{
+	if (!Delegates.IsValid())
+	{
+		return;
+	}
 
-	Delegates = InDelegates;
+	BindCommands();
 
 	ChildSlot
 	[
@@ -75,7 +83,7 @@ void SMediaImageViewerOverlay::Construct(const FArguments& InArgs, EMediaImageVi
 			.HAlign(HAlign_Fill)
 			.VAlign(VAlign_Center)
 			[
-				CreateStatusBar(InDelegates)
+				CreateStatusBar(Delegates.ToSharedRef())
 			]
 
 			+ SVerticalBox::Slot()
@@ -83,7 +91,7 @@ void SMediaImageViewerOverlay::Construct(const FArguments& InArgs, EMediaImageVi
 			.HAlign(HAlign_Fill)
 			.VAlign(VAlign_Fill)
 			[
-				CreateOverlay(InArgs._bComparisonView)
+				CreateOverlay(bComparisonView)
 			]
 		]
 	];
@@ -378,7 +386,12 @@ void SMediaImageViewerOverlay::UpdateDragPosition()
 	NewOffset.X += CursorLocation.X - DragStartCursor.X;
 	NewOffset.Y += CursorLocation.Y - DragStartCursor.Y;
 
-	Try_AddOffset(NewOffset - GetOffset());
+	const FVector OffsetDiff = NewOffset - GetOffset();
+
+	if (!OffsetDiff.IsNearlyZero())
+	{
+		Try_AddOffset(OffsetDiff);
+	}
 }
 
 void SMediaImageViewerOverlay::UpdateDragScale()
@@ -399,6 +412,11 @@ void SMediaImageViewerOverlay::UpdateDragScale()
 	const FVector2D CursorLocation = Delegates->GetCursorLocation.Execute();
 	const float Difference = (CursorLocation - DragStartCursor).Y;
 	const float NewScale = LastDragScaleValue - Difference * ModifierScale * DragMultiplier;
+
+	if (FMath::IsNearlyEqual(NewScale, GetScale()))
+	{
+		return;
+	}
 	
 	// Since we can change speed of the scale change mid operation, reset this each update.
 	DragStartCursor = CursorLocation;
@@ -494,7 +512,7 @@ TSharedRef<SWidget> SMediaImageViewerOverlay::CreateOverlay(bool bInComparisonVi
 	Overlay = SNew(SOverlay);
 
 	TSharedPtr<SWidget> ImageViewerOverlay = ImageViewer.IsValid() 
-		? ImageViewer->GetOverlayWidget(Position, Delegates->GetTab.Execute()) 
+		? ImageViewer->GetOverlayWidget(Position, Delegates) 
 		: nullptr;
 
 	Overlay->AddSlot()
@@ -954,7 +972,7 @@ void SMediaImageViewerOverlay::AddToLibrary()
 		return;
 	}
 
-	TSharedRef<FMediaViewerLibrary> Library = Delegates->GetLibrary.Execute();
+	TSharedRef<IMediaViewerLibrary> Library = Delegates->GetLibrary.Execute();
 
 	Library->AddItemToGroup(CachedItem.ToSharedRef());
 }
diff --git a/Engine/Plugins/Experimental/MediaViewer/Source/MediaViewer/Private/Widgets/SMediaImageViewerOverlay.h b/Engine/Plugins/Experimental/MediaViewer/Source/MediaViewer/Private/Widgets/SMediaImageViewerOverlay.h
--- a/Engine/Plugins/Experimental/MediaViewer/Source/MediaViewer/Private/Widgets/SMediaImageViewerOverlay.h
+++ b/Engine/Plugins/Experimental/MediaViewer/Source/MediaViewer/Private/Widgets/SMediaImageViewerOverlay.h
@@ -4,8 +4,6 @@
 
 #include "Widgets/SCompoundWidget.h"
 
-#include "Widgets/MediaViewerDelegates.h"
-
 class FUICommandList;
 class SOverlay;
 class SVerticalBox;
@@ -15,6 +13,7 @@ namespace UE::MediaViewer
 {
 	class FMediaImageViewer;
 	enum class EMediaImageViewerPosition : uint8;
+	struct FMediaViewerDelegates;
 }
 
 namespace UE::MediaViewer::Private
@@ -44,6 +43,8 @@ public:
 
 	void Construct(const FArguments& InArgs, EMediaImageViewerPosition InPosition, const TSharedRef<FMediaViewerDelegates>& InDelegates);
 
+	void PostConstruct();
+
 	TSharedPtr<FMediaImageViewer> GetImageViewer() const;
 
 	const TSharedRef<FUICommandList>& GetCommandList() const;
@@ -95,6 +96,7 @@ protected:
 	TSharedPtr<FMediaViewerDelegates> Delegates;
 	TSharedPtr<FMediaViewerLibraryItem> CachedItem;
 	TSharedRef<FUICommandList> CommandList;
+	bool bComparisonView;
 	/** Name of the button which is currently performing a drag operation. */
 	FName DragButtonName;
 	EDragType DragType;
diff --git a/Engine/Plugins/Experimental/MediaViewer/Source/MediaViewer/Private/Widgets/SMediaImageViewerStatusBar.h b/Engine/Plugins/Experimental/MediaViewer/Source/MediaViewer/Private/Widgets/SMediaImageViewerStatusBar.h
--- a/Engine/Plugins/Experimental/MediaViewer/Source/MediaViewer/Private/Widgets/SMediaImageViewerStatusBar.h
+++ b/Engine/Plugins/Experimental/MediaViewer/Source/MediaViewer/Private/Widgets/SMediaImageViewerStatusBar.h
@@ -4,7 +4,7 @@
 
 #include "Widgets/SCompoundWidget.h"
 
-#include "Widgets/MediaViewerDelegates.h"
+#include "MediaViewerDelegates.h"
 
 class FText;
 class FUICommandList;
diff --git a/Engine/Plugins/Experimental/MediaViewer/Source/MediaViewer/Private/Widgets/SMediaSourceOverlay.cpp b/Engine/Plugins/Experimental/MediaViewer/Source/MediaViewer/Private/Widgets/SMediaSourceOverlay.cpp
--- a/Engine/Plugins/Experimental/MediaViewer/Source/MediaViewer/Private/Widgets/SMediaSourceOverlay.cpp
+++ b/Engine/Plugins/Experimental/MediaViewer/Source/MediaViewer/Private/Widgets/SMediaSourceOverlay.cpp
@@ -7,6 +7,8 @@
 #include "IMediaStreamPlayer.h"
 #include "MediaPlayer.h"
 #include "MediaStream.h"
+#include "MediaViewerCommands.h"
+#include "MediaViewerDelegates.h"
 #include "MediaViewerStyle.h"
 #include "Misc/FrameRate.h"
 #include "Misc/Timecode.h"
@@ -19,7 +21,6 @@
 #include "Widgets/Input/SButton.h"
 #include "Widgets/Layout/SBorder.h"
 #include "Widgets/Layout/SBox.h"
-#include "Widgets/MediaViewerDelegates.h"
 #include "Widgets/SBoxPanel.h"
 #include "Widgets/SMediaViewer.h"
 #include "Widgets/SMediaViewerTab.h"
@@ -56,16 +57,19 @@ void SMediaSourceOverlay::PrivateRegisterAttributes(FSlateAttributeDescriptor::F
 }
 
 void SMediaSourceOverlay::Construct(const FArguments& InArgs, const TSharedRef<FMediaSourceImageViewer>& InImageViewer,
-	const TSharedPtr<SMediaViewerTab>& InViewerTab)
+	EMediaImageViewerPosition InPosition, const TSharedPtr<FMediaViewerDelegates>& InDelegates)
 {
 	ImageViewerWeak = InImageViewer;
+	Position = InPosition;
 
-	if (InViewerTab.IsValid())
+	if (InDelegates.IsValid())
 	{
-		Delegates = InViewerTab->GetViewer()->GetDelegates();
+		Delegates = InDelegates;
 		Delegates->ImageViewerEvent.AddSP(this, &SMediaSourceOverlay::ReceiveEvent);
 	}
 
+	BindCommands();
+
 	TrySetFrameRate();
 
 	ChildSlot
@@ -168,6 +172,31 @@ void SMediaSourceOverlay::TrySetFrameRate()
 	}
 }
 
+void SMediaSourceOverlay::BindCommands()
+{
+	if (!Delegates.IsValid() || !Delegates->GetCommandListForPosition.IsBound())
+	{
+		return;
+	}
+
+	TSharedPtr<FUICommandList> CommandList = Delegates->GetCommandListForPosition.Execute(Position);
+
+	if (!CommandList.IsValid())
+	{
+		return;
+	}
+
+	const FMediaViewerCommands& Commands = FMediaViewerCommands::Get();
+
+	CommandList->MapAction(Commands.QuickRewindVideo, FExecuteAction::CreateSP(this, &SMediaSourceOverlay::AddOffset_Time, -10.f));
+	CommandList->MapAction(Commands.RewindVideo, FExecuteAction::CreateSP(this, &SMediaSourceOverlay::AddOffset_Time, -1.f));
+	CommandList->MapAction(Commands.StepBackVideo, FExecuteAction::CreateSP(this, &SMediaSourceOverlay::AddOffset_Frame, -1));
+	CommandList->MapAction(Commands.ToggleVideoPlay, FExecuteAction::CreateSP(this, &SMediaSourceOverlay::TogglePlay));
+	CommandList->MapAction(Commands.StepForwardVideo, FExecuteAction::CreateSP(this, &SMediaSourceOverlay::AddOffset_Frame, 1));
+	CommandList->MapAction(Commands.FastForwardVideo, FExecuteAction::CreateSP(this, &SMediaSourceOverlay::AddOffset_Time, 1.f));
+	CommandList->MapAction(Commands.QuickFastForwardVideo, FExecuteAction::CreateSP(this, &SMediaSourceOverlay::AddOffset_Time, 10.f));
+}
+
 TSharedRef<SWidget> SMediaSourceOverlay::CreateSlider()
 {
 	UMediaPlayer* MediaPlayer = GetMediaPlayer();
@@ -1131,6 +1160,77 @@ bool SMediaSourceOverlay::Command_Forward(UMediaPlayer* InMediaPlayer)
 	return false;
 }
 
+void SMediaSourceOverlay::AddOffset_Time(float InOffset)
+{
+	UMediaPlayer* MediaPlayer = GetMediaPlayer();
+
+	if (!MediaPlayer)
+	{
+		return;
+	}
+
+	const float Duration = MediaPlayer->GetDuration().GetTotalSeconds();
+
+	if (FMath::IsNearlyZero(Duration))
+	{
+		return;
+	}
+
+	const float CurrentPosition = MediaPlayer->GetTime().GetTotalSeconds();
+
+	if (CurrentPosition < 0 || CurrentPosition > Duration)
+	{
+		return;
+	}
+
+	const float NewPosition = FMath::Clamp(CurrentPosition + InOffset, 0, Duration);
+
+	if (FMath::IsNearlyEqual(CurrentPosition, NewPosition))
+	{
+		return;
+	}
+
+	MediaPlayer->Seek(FTimespan::FromSeconds(NewPosition));
+
+	// Broadcast the event if required by simulating a scrub.
+	const float StartFraction = CurrentPosition / Duration;
+	const float EndFraction = NewPosition / Duration;
+
+	TArray<UMediaPlayer*, TInlineAllocator<1>> MediaPlayers = {MediaPlayer};	
+
+	OnScrub(IMediaPlayerSlider::EScrubEventType::Begin, MediaPlayers, StartFraction);
+	OnScrub(IMediaPlayerSlider::EScrubEventType::Update, MediaPlayers, EndFraction);
+	OnScrub(IMediaPlayerSlider::EScrubEventType::End, MediaPlayers, EndFraction);
+}
+
+void SMediaSourceOverlay::AddOffset_Frame(int32 InOffset)
+{
+	switch (InOffset)
+	{
+		case 1:
+			if (StepForward_IsEnabled())
+			{
+				StepForward_OnClicked();
+			}
+			break;
+
+		case -1:
+			if (StepBack_IsEnabled())
+			{
+				StepBack_OnClicked();
+			}
+			break;
+	}
+}
+
+void SMediaSourceOverlay::TogglePlay()
+{
+	if (Play_IsEnabled())
+	{
+		Play_OnClicked();
+	}
+}
+
 }
 
 #undef LOCTEXT_NAMESPACE
diff --git a/Engine/Plugins/Experimental/MediaViewer/Source/MediaViewer/Private/Widgets/SMediaSourceOverlay.h b/Engine/Plugins/Experimental/MediaViewer/Source/MediaViewer/Private/Widgets/SMediaSourceOverlay.h
--- a/Engine/Plugins/Experimental/MediaViewer/Source/MediaViewer/Private/Widgets/SMediaSourceOverlay.h
+++ b/Engine/Plugins/Experimental/MediaViewer/Source/MediaViewer/Private/Widgets/SMediaSourceOverlay.h
@@ -18,14 +18,15 @@ struct FSlateBrush;
 namespace UE::MediaViewer
 {
 	class SMediaViewerTab;
+	enum class EMediaImageViewerPosition : uint8;
+	struct FMediaImageViewerEventParams;
+	struct FMediaViewerDelegates;
 }
 
 namespace UE::MediaViewer::Private
 {
 
 class FMediaSourceImageViewer;
-struct FMediaImageViewerEventParams;
-struct FMediaViewerDelegates;
 
 namespace MediaSourceOverlay::Events
 {
@@ -43,20 +44,33 @@ class SMediaSourceOverlay : public SCompoundWidget
 public:
 	virtual ~SMediaSourceOverlay() override = default;
 
-	void Construct(const FArguments& InArgs, const TSharedRef<FMediaSourceImageViewer>& InImageViewer, 
-		const TSharedPtr<SMediaViewerTab>& InViewerTab);
+	void Construct(const FArguments& InArgs, const TSharedRef<FMediaSourceImageViewer>& InImageViewer, EMediaImageViewerPosition InPosition,
+		const TSharedPtr<FMediaViewerDelegates>& InDelegates);
 
 	//~ Begin SWidget
 	virtual void Tick(const FGeometry& InAllottedGeometry, const double InCurrentTime, const float InDeltaTime) override;
 	//~ End SWidget
 
 protected:
+	/** The image viewer being controlled by this overlay. */
 	TWeakPtr<FMediaSourceImageViewer> ImageViewerWeak;
+
+	/** The position in the viewer occupied by this image viewer. */
+	EMediaImageViewerPosition Position;
+
+	/** The delegates used to communicate with other widgets. */
 	TSharedPtr<FMediaViewerDelegates> Delegates;
+
+	/** Container so the overlay can be toggled on and off. */
 	TSharedPtr<SBox> Container;
+
+	/** The frame rate stored before a scrub-based pause event. */
 	TOptional<float> FrameRateFloat;
+
 	mutable TOptional<FText> TotalFrames;
 	mutable TOptional<FText> TotalTime;
+
+	/** When a sync scrub is started, this stores the start of this media. */
 	TVariant<FTimespan, float> SyncedScrubStartTime;
 
 	/** The last time the mouse was moved over the widget. */
@@ -68,10 +82,14 @@ protected:
 
 	UMediaPlayer* GetMediaPlayer() const;
 
+	/** Attempt to store the frame rate of the currently playing media. */
 	void TrySetFrameRate();
 
+	void BindCommands();
+
 	TSharedRef<SWidget> CreateSlider();
 
+	/** Create play controls. Play, pause, etc. */
 	TSharedRef<SWidget> CreateControls();
 
 	FText GetCurrentFrame() const;
@@ -102,12 +120,16 @@ protected:
 	bool Forward_IsEnabled() const;
 	FReply Forward_OnClicked();
 
+	/** Broadcasts an event that can be received by others. */
 	void SendEvent(FName InEventName);
 
+	/** Receive an event broadcast by somebody else. */
 	void ReceiveEvent(const FMediaImageViewerEventParams& InEventParams);
 
+	/** Triggered when this player's scrub bar is manually moved. */
 	void OnScrub(IMediaPlayerSlider::EScrubEventType InScrubEvent, TConstArrayView<UMediaPlayer*> InMediaPlayers, float InScrubPosition);
 
+	/** Execute commands on the given media player. */
 	static bool Command_Rewind(UMediaPlayer* InMediaPlayer);
 	static bool Command_Reverse(UMediaPlayer* InMediaPlayer);
 	static bool Command_StepBack(UMediaPlayer* InMediaPlayer);
@@ -116,7 +138,17 @@ protected:
 	static bool Command_StepForward(UMediaPlayer* InMediaPlayer);
 	static bool Command_Forward(UMediaPlayer* InMediaPlayer);
 	
+	/** Triggered when a scrub event is received from someone else. */
 	void ReceiveScrub(const MediaSourceOverlay::Events::FMediaImageViewerScrubEventParams& InEventParams);
+
+	/** Add a offset to the current media. Optionally broadcasts it to others as a scrub event. */
+	void AddOffset_Time(float InOffset);
+
+	/** Add a offset to the current media. Optionally broadcasts it to others as a scrub event. Supports 1 and -1. */
+	void AddOffset_Frame(int32 InOffset);
+
+	/** Toggles between play and pause. Forward only. */
+	void TogglePlay();
 };
 
 } // UE::MediaViewer::Private
\ No newline at end of file
diff --git a/Engine/Plugins/Experimental/MediaViewer/Source/MediaViewer/Private/Widgets/SMediaViewer.cpp b/Engine/Plugins/Experimental/MediaViewer/Source/MediaViewer/Private/Widgets/SMediaViewer.cpp
--- a/Engine/Plugins/Experimental/MediaViewer/Source/MediaViewer/Private/Widgets/SMediaViewer.cpp
+++ b/Engine/Plugins/Experimental/MediaViewer/Source/MediaViewer/Private/Widgets/SMediaViewer.cpp
@@ -22,21 +22,21 @@
 #include "Library/MediaViewerLibraryGroup.h"
 #include "Library/MediaViewerLibraryIni.h"
 #include "MediaViewerCommands.h"
+#include "MediaViewerDelegates.h"
 #include "MediaViewerUtils.h"
 #include "Misc/MessageDialog.h"
 #include "Sidebar/SSidebar.h"
 #include "Sidebar/SSidebarContainer.h"
 #include "TextureResource.h"
+#include "UObject/StrongObjectPtr.h"
 #include "Widgets/Colors/SColorBlock.h"
 #include "Widgets/Layout/SBox.h"
 #include "Widgets/Layout/SSplitter.h"
-#include "Widgets/MediaViewerDelegates.h"
 #include "Widgets/Notifications/SNotificationList.h"
 #include "Widgets/SMediaImageViewerOverlay.h"
 #include "Widgets/SMediaViewerLibrary.h"
 #include "Widgets/SMediaViewerLibraryPrivate.h"
 #include "Widgets/SMediaViewerToolbar.h"
-#include "UObject/StrongObjectPtr.h"
 
 #include UE_INLINE_GENERATED_CPP_BY_NAME(SMediaViewer)
 
@@ -555,7 +555,6 @@ void SMediaViewer::CreateDelegates()
 	Delegates->RefreshView.BindSPLambda(this, [this] { InvalidateView(); });
 	Delegates->GetCommandList.BindSPLambda(this, [this] { return CommandList; });
 	Delegates->IsOverViewer.BindSPLambda(this, [this] { return IsHovered(); });
-	Delegates->GetTab.BindSPLambda(this, [this] { return Tab.Pin(); });
 
 	// Per panel delegates
 	Delegates->GetImageViewer.BindSP(this, &SMediaViewer::GetImageViewer);
@@ -622,14 +621,21 @@ void SMediaViewer::CreateView()
 
 void SMediaViewer::CreateSingleView(EMediaImageViewerPosition InPosition)
 {
+	TSharedRef<SMediaImageViewerOverlay> Overlay = CreateOverlay(InPosition, /* Show Panel Name */ false);
+
 	// Attach the single overlay
-	ContentSlot->AttachWidget(CreateOverlay(InPosition, /* Show Panel Name */ false));
+	ContentSlot->AttachWidget(Overlay);
+
+	Overlay->PostConstruct();
 }
 
 void SMediaViewer::CreateABView()
 {
 	MediaViewerSettings.ABSplitterLocation = 50.f;
 
+	TSharedRef<SMediaImageViewerOverlay> First = CreateOverlay(EMediaImageViewerPosition::First, /* Show Panel Name */ true);
+	TSharedRef<SMediaImageViewerOverlay> Second = CreateOverlay(EMediaImageViewerPosition::Second, /* Show Panel Name */ true);
+
 	ContentSlot->AttachWidget(
 		SNew(SSplitter)
 		.Orientation(MediaViewerSettings.ABOrientation)
@@ -643,7 +649,7 @@ void SMediaViewer::CreateABView()
 		.SizeRule(SSplitter::FractionOfParent)
 		.OnSlotResized(this, &SMediaViewer::OnABResized)
 		[
-			CreateOverlay(EMediaImageViewerPosition::First, /* Show Panel Name */ true)
+			First
 		]
 
 		+ SSplitter::Slot()
@@ -652,9 +658,12 @@ void SMediaViewer::CreateABView()
 		.Resizable(true)
 		.SizeRule(SSplitter::FractionOfParent)
 		[
-			CreateOverlay(EMediaImageViewerPosition::Second, /* Show Panel Name */ true)
+			Second
 		]
-	);	
+	);
+
+	First->PostConstruct();
+	Second->PostConstruct();
 }
 
 TSharedPtr<SMediaImageViewerOverlay> SMediaViewer::GetOverlay(EMediaImageViewerPosition InPosition) const
diff --git a/Engine/Plugins/Experimental/MediaViewer/Source/MediaViewer/Private/Widgets/SMediaViewerDropTarget.cpp b/Engine/Plugins/Experimental/MediaViewer/Source/MediaViewer/Private/Widgets/SMediaViewerDropTarget.cpp
--- a/Engine/Plugins/Experimental/MediaViewer/Source/MediaViewer/Private/Widgets/SMediaViewerDropTarget.cpp
+++ b/Engine/Plugins/Experimental/MediaViewer/Source/MediaViewer/Private/Widgets/SMediaViewerDropTarget.cpp
@@ -12,11 +12,11 @@
 #include "Library/MediaViewerLibrary.h"
 #include "Library/MediaViewerLibraryItem.h"
 #include "MediaPlayer.h"
+#include "MediaViewerDelegates.h"
 #include "SDropTarget.h"
 #include "Styling/AppStyle.h"
 #include "Styling/StyleColors.h"
 #include "Widgets/Layout/SBox.h"
-#include "Widgets/MediaViewerDelegates.h"
 #include "Widgets/MediaViewerLibraryItemDragDropOperation.h"
 #include "Widgets/SBoxPanel.h"
 #include "Widgets/SOverlay.h"
@@ -164,7 +164,7 @@ FReply SMediaViewerDropTarget::OnDropped(const FGeometry& InGeometry, const FDra
 
 void SMediaViewerDropTarget::HandleDroppedMediaViewerOp(const FMediaViewerLibraryItemDragDropOperation& InMediaViewerOp)
 {
-	TSharedRef<FMediaViewerLibrary> Library = Delegates->GetLibrary.Execute();
+	TSharedRef<IMediaViewerLibrary> Library = Delegates->GetLibrary.Execute();
 
 	if (TSharedPtr<FMediaViewerLibraryItem> LibraryItem = Library->GetItem(InMediaViewerOp.GetGroupItem().ItemId))
 	{
diff --git a/Engine/Plugins/Experimental/MediaViewer/Source/MediaViewer/Private/Widgets/SMediaViewerDropTarget.h b/Engine/Plugins/Experimental/MediaViewer/Source/MediaViewer/Private/Widgets/SMediaViewerDropTarget.h
--- a/Engine/Plugins/Experimental/MediaViewer/Source/MediaViewer/Private/Widgets/SMediaViewerDropTarget.h
+++ b/Engine/Plugins/Experimental/MediaViewer/Source/MediaViewer/Private/Widgets/SMediaViewerDropTarget.h
@@ -6,11 +6,15 @@
 
 #include "SMediaViewer.h"
 
+namespace UE::MediaViewer
+{
+	struct FMediaViewerDelegates;
+}
+
 namespace UE::MediaViewer::Private
 {
 
 class FMediaViewerLibraryItemDragDropOperation;
-struct FMediaViewerDelegates;
 
 class SMediaViewerDropTarget : public SCompoundWidget
 {
diff --git a/Engine/Plugins/Experimental/MediaViewer/Source/MediaViewer/Private/Widgets/SMediaViewerLibraryItem.cpp b/Engine/Plugins/Experimental/MediaViewer/Source/MediaViewer/Private/Widgets/SMediaViewerLibraryItem.cpp
--- a/Engine/Plugins/Experimental/MediaViewer/Source/MediaViewer/Private/Widgets/SMediaViewerLibraryItem.cpp
+++ b/Engine/Plugins/Experimental/MediaViewer/Source/MediaViewer/Private/Widgets/SMediaViewerLibraryItem.cpp
@@ -12,6 +12,7 @@
 #include "Library/MediaViewerLibraryGroup.h"
 #include "Library/MediaViewerLibraryItem.h"
 #include "MediaViewer.h"
+#include "MediaViewerDelegates.h"
 #include "MediaViewerStyle.h"
 #include "Styling/AppStyle.h"
 #include "Styling/StyleColors.h"
@@ -19,7 +20,6 @@
 #include "Widgets/Images/SImage.h"
 #include "Widgets/Input/SButton.h"
 #include "Widgets/Input/SCheckBox.h"
-#include "Widgets/MediaViewerDelegates.h"
 #include "Widgets/MediaViewerLibraryItemAssetDragDropOperation.h"
 #include "Widgets/MediaViewerLibraryItemDragDropOperation.h"
 #include "Widgets/SBoxPanel.h"
diff --git a/Engine/Plugins/Experimental/MediaViewer/Source/MediaViewer/Private/Widgets/SMediaViewerLibraryPrivate.cpp b/Engine/Plugins/Experimental/MediaViewer/Source/MediaViewer/Private/Widgets/SMediaViewerLibraryPrivate.cpp
--- a/Engine/Plugins/Experimental/MediaViewer/Source/MediaViewer/Private/Widgets/SMediaViewerLibraryPrivate.cpp
+++ b/Engine/Plugins/Experimental/MediaViewer/Source/MediaViewer/Private/Widgets/SMediaViewerLibraryPrivate.cpp
@@ -12,8 +12,8 @@
 #include "Library/MediaViewerLibraryIni.h"
 #include "Library/MediaViewerLibraryItem.h"
 #include "MediaViewer.h"
+#include "MediaViewerDelegates.h"
 #include "Misc/Optional.h"
-#include "Widgets/MediaViewerDelegates.h"
 #include "Widgets/SMediaViewer.h"
 #include "Widgets/SMediaViewerLibrary.h"
 
diff --git a/Engine/Plugins/Experimental/MediaViewer/Source/MediaViewer/Private/Widgets/SMediaViewerToolbar.h b/Engine/Plugins/Experimental/MediaViewer/Source/MediaViewer/Private/Widgets/SMediaViewerToolbar.h
--- a/Engine/Plugins/Experimental/MediaViewer/Source/MediaViewer/Private/Widgets/SMediaViewerToolbar.h
+++ b/Engine/Plugins/Experimental/MediaViewer/Source/MediaViewer/Private/Widgets/SMediaViewerToolbar.h
@@ -5,20 +5,19 @@
 #include "Misc/NotifyHook.h"
 #include "SViewportToolBar.h"
 
-#include "Widgets/MediaViewerDelegates.h"
+#include "MediaViewerDelegates.h"
 
 class IStructureDetailsView;
 
 namespace UE::MediaViewer
 {
 	class FMediaViewer;
+	struct FMediaViewerDelegates;
 }
 
 namespace UE::MediaViewer::Private
 {
 
-struct FMediaViewerDelegates;
-
 /**
  * Toolbar for the Media Viewer.
  */
diff --git a/Engine/Plugins/Experimental/MediaViewer/Source/MediaViewer/Public/ImageViewer/MediaImageViewer.h b/Engine/Plugins/Experimental/MediaViewer/Source/MediaViewer/Public/ImageViewer/MediaImageViewer.h
--- a/Engine/Plugins/Experimental/MediaViewer/Source/MediaViewer/Public/ImageViewer/MediaImageViewer.h
+++ b/Engine/Plugins/Experimental/MediaViewer/Source/MediaViewer/Public/ImageViewer/MediaImageViewer.h
@@ -108,7 +108,7 @@ namespace UE::MediaViewer
 {
 
 class FMediaImageStatusBarExtender;
-class SMediaViewerTab;
+struct FMediaViewerDelegates;
 	
 /** Information about the image viewer's image. */
 struct FMediaImageViewerInfo
@@ -257,7 +257,7 @@ public:
 	 * The above alignment will be used for the widget's slot.
 	 * If no overlay is required, return nullptr.
 	 */
-	virtual TSharedPtr<SWidget> GetOverlayWidget(EMediaImageViewerPosition InPosition, const TSharedPtr<SMediaViewerTab>& InViewerTab) { return nullptr; }
+	virtual TSharedPtr<SWidget> GetOverlayWidget(EMediaImageViewerPosition InPosition, const TSharedPtr<FMediaViewerDelegates>& InDelegates) { return nullptr; }
 
 	/** 
 	 * Responds to tracking started in this viewer's viewport. 
diff --git a/Engine/Plugins/Experimental/MediaViewer/Source/MediaViewer/Public/MediaViewerDelegates.h b/Engine/Plugins/Experimental/MediaViewer/Source/MediaViewer/Public/MediaViewerDelegates.h
new file mode 100644
--- /dev/null
+++ b/Engine/Plugins/Experimental/MediaViewer/Source/MediaViewer/Public/MediaViewerDelegates.h
@@ -0,0 +1,166 @@
+// Copyright Epic Games, Inc. All Rights Reserved.
+
+#pragma once
+
+#include "Delegates/Delegate.h"
+#include "Delegates/DelegateCombinations.h"
+#include "Math/MathFwd.h"
+#include "MediaViewer.h"
+#include "Templates/SharedPointerFwd.h"
+#include "UObject/StructOnScope.h"
+
+class FUICommandList;
+enum class EMediaImageViewerActivePosition : uint8;
+enum EOrientation : int;
+struct FMediaViewerSettings;
+
+namespace UE::MediaViewer
+{
+
+class FMediaImageViewer;
+class IMediaViewerLibrary;
+class SMediaViewerTab;
+	
+struct FMediaImageViewerEventParams
+{
+	TSharedPtr<FMediaImageViewer> FromViewer;
+	FName EventName;
+};
+
+struct FMediaViewerDelegates
+{
+	DECLARE_DELEGATE_RetVal(FVector2D, FGetLocation)
+	DECLARE_DELEGATE_RetVal(EOrientation, FGetOrientation)
+	DECLARE_DELEGATE_OneParam(FSetOrientation, EOrientation)
+	DECLARE_DELEGATE_RetVal(EMediaImageViewerActivePosition, FGetActiveView)
+	DECLARE_DELEGATE_RetVal(bool, FGetBool)
+	DECLARE_DELEGATE_RetVal(float, FGetFloat)
+	DECLARE_DELEGATE_OneParam(FSetFloat, float)
+	DECLARE_DELEGATE_OneParam(FAddOffset, const FVector&)
+	DECLARE_DELEGATE_OneParam(FAddRotation, const FRotator&)
+	DECLARE_DELEGATE_OneParam(FMultiplyScale, float)
+	DECLARE_DELEGATE_ThreeParams(FSetTransform, const FVector&, const FRotator&, float)
+	DECLARE_DELEGATE_RetVal(FMediaViewerSettings*, FGetSettings)
+	DECLARE_DELEGATE_RetVal(TSharedRef<IMediaViewerLibrary>, FGetLibrary)
+	DECLARE_DELEGATE_RetVal(TSharedPtr<FUICommandList>, FGetCommandList)
+
+	DECLARE_DELEGATE_RetVal_OneParam(TSharedPtr<FMediaImageViewer>, FGetImageViewer, EMediaImageViewerPosition)
+	DECLARE_DELEGATE_TwoParams(FSetImageViewer, EMediaImageViewerPosition, const TSharedRef<FMediaImageViewer>&)
+	DECLARE_DELEGATE_OneParam(FSimpleForPosition, EMediaImageViewerPosition)
+	DECLARE_DELEGATE_RetVal_OneParam(FIntPoint, FGetLocationForPosition, EMediaImageViewerPosition)
+	DECLARE_DELEGATE_RetVal_OneParam(bool, FGetBoolForPosition, EMediaImageViewerPosition)
+	DECLARE_DELEGATE_RetVal_OneParam(TSharedPtr<FUICommandList>, FGetCommandListForPosition, EMediaImageViewerPosition)
+	DECLARE_MULTICAST_DELEGATE_OneParam(FImageViewerEvent, const FMediaImageViewerEventParams&)
+
+	//////////////////////////////
+	// General viewer delegates //
+	//////////////////////////////
+
+	/** Changes to single view mode. */
+	FSimpleDelegate SetSingleView;
+
+	/** Changes to AB view. */
+	FSimpleDelegate SetABView;
+	
+	/** Gets the orientation in A/B mode. Will return the currently selected orientation even if not in A/B mode. */
+	FGetOrientation GetABOrientation;
+
+	/** Sets the orientation in A/B mode. Will set the next expected orientation even if not in A/B mode. */
+	FSetOrientation SetABOrientation;
+
+	/** Returns a bitflag value indicating which image viewers are currently active (have images, not null or null image viewer. */
+	FGetActiveView GetActiveView;
+
+	/* Gets a pointer to the Media Viewer configurable settings. */
+	FGetSettings GetSettings;
+	
+	/** Gets the size of the viewer paint area. */
+	FGetLocation GetViewerSize;
+
+	/** Gets the position of the viewer paint area in the window. */
+	FGetLocation GetViewerPosition;
+
+	/** Gets the local cursor location on the window. */
+	FGetLocation GetCursorLocation;
+
+	/** Swaps the A and B viewers, including their transforms. */
+	FSimpleDelegate SwapAB;
+
+	/** Returns true if the transforms of multi-image-viewer views are locked in sync. */
+	FGetBool AreTransformsLocked;
+
+	/** Toggles the transform lock for all image viewers. */
+	FSimpleDelegate ToggleLockedTransform;
+
+	/** Add an offset to all image viewers. */
+	FAddOffset AddOffsetToAll;
+
+	/** Adds a rotation to all image viewers. */
+	FAddRotation AddRotationToAll;
+
+	/** Multiplies the scale of all image viewers. */
+	FMultiplyScale MultiplyScaleToAll;
+
+	/** Multiplies the scale of all image viewers. */
+	FMultiplyScale MultiplyScaleAroundCursorToAll;
+
+	/** Sets the transform of all image viewers. */
+	FSetTransform SetTransformToAll;
+
+	/** Resets the transform of all image viewers to 0,0,0; 0,0,0; 1. */
+	FSimpleDelegate ResetTransformToAll;
+
+	/** Gets the opacity of the second/B image viewer. */
+	FGetFloat GetSecondImageViewerOpacity;
+
+	/** Sets the opacity of the second/B image viewer. */
+	FSetFloat SetSecondImageViewerOpacity;
+
+	/** Gets the splitter location in percent (0-100) */
+	FGetFloat GetABSplitterLocation;
+
+	/** Sets the splitter location in percent (0-100) */
+	FSetFloat SetABSplitterLocation;
+
+	/** Gets the image viewer library. */
+	FGetLibrary GetLibrary;
+
+	/* Forces a refresh of the view on next tick. */
+	FSimpleDelegate RefreshView;
+
+	/** Returns the command list for the general view. */
+	FGetCommandList GetCommandList;
+
+	/** Returns true if the mouse is over the viewer. */
+	FGetBool IsOverViewer;
+
+	//////////////////////////
+	// Per-viewer delegates //
+	//////////////////////////
+
+	/** Gets the image viewer in the given position. */
+	FGetImageViewer GetImageViewer;
+
+	/** Sets the image viewer in the given position. */
+	FSetImageViewer SetImageViewer;
+
+	/** Sets the image viewer in the given position to the null viewer. */
+	FSimpleForPosition ClearImageViewer;
+
+	/** Gets the pixel coordinates hovered by the mouse for the image viewer in the given position. */
+	FGetLocationForPosition GetPixelCoordinates;
+
+	/** Returns true if the mouse is over the image viewer in the given position. */
+	FGetBoolForPosition IsOverImage;
+
+	/** Copies the transform of the image viewer in the given position to all other image viewers. */
+	FSimpleForPosition CopyTransformToAll;
+
+	/** Gets the command list for image viewer in the given position. */
+	FGetCommandListForPosition GetCommandListForPosition;
+
+	/** Generic event that can be broadcast between image viewers and controls. */
+	FImageViewerEvent ImageViewerEvent;
+};
+
+} // UE::MediaViewer
