Media Viewer: When locked scrubbing, the videos not being directly scrubbed is paused.

diff --git a/Engine/Plugins/Experimental/MediaViewer/Source/MediaViewer/Private/Widgets/SMediaSourceOverlay.cpp b/Engine/Plugins/Experimental/MediaViewer/Source/MediaViewer/Private/Widgets/SMediaSourceOverlay.cpp
--- a/Engine/Plugins/Experimental/MediaViewer/Source/MediaViewer/Private/Widgets/SMediaSourceOverlay.cpp
+++ b/Engine/Plugins/Experimental/MediaViewer/Source/MediaViewer/Private/Widgets/SMediaSourceOverlay.cpp
@@ -61,6 +61,7 @@ void SMediaSourceOverlay::Construct(const FArguments& InArgs, const TSharedRef<F
 {
 	ImageViewerWeak = InImageViewer;
 	Position = InPosition;
+	RateWhenScrubEventReceived = {};
 
 	if (InDelegates.IsValid())
 	{
@@ -538,10 +539,10 @@ FReply SMediaSourceOverlay::Rewind_OnClicked()
 {
 	if (UMediaPlayer* MediaPlayer = GetMediaPlayer())
 	{
-		MediaPlayer->Pause();
-		MediaPlayer->Rewind();
-
-		SendEvent(MediaSourceOverlay::Events::Rewind);
+		if (Command_Pause(MediaPlayer) && Command_Rewind(MediaPlayer))
+		{
+			SendEvent(MediaSourceOverlay::Events::Rewind);
+		}
 	}
 
 	return FReply::Handled();
@@ -1037,6 +1038,12 @@ void SMediaSourceOverlay::ReceiveScrub(const MediaSourceOverlay::Events::FMediaI
 	{
 		case IMediaPlayerSlider::EScrubEventType::Begin:
 		{
+			if (MediaPlayer->IsPlaying())
+			{
+				RateWhenScrubEventReceived = MediaPlayer->GetRate();
+				Command_Pause(MediaPlayer);
+			}
+
 			switch (MediaSyncType)
 			{
 				case EMediaViewerMediaSyncType::RelativeTime:
@@ -1073,6 +1080,13 @@ void SMediaSourceOverlay::ReceiveScrub(const MediaSourceOverlay::Events::FMediaI
 		{
 			UpdateScrub();
 			SyncedScrubStartTime = {};
+
+			if (RateWhenScrubEventReceived.IsSet())
+			{
+				MediaPlayer->SetRate(RateWhenScrubEventReceived.GetValue());
+				RateWhenScrubEventReceived = {};
+			}
+
 			break;
 		}
 
@@ -1084,7 +1098,7 @@ void SMediaSourceOverlay::ReceiveScrub(const MediaSourceOverlay::Events::FMediaI
 
 bool SMediaSourceOverlay::Command_Rewind(UMediaPlayer* InMediaPlayer)
 {
-	InMediaPlayer->Pause();
+	InMediaPlayer->SetRate(0.f);
 	InMediaPlayer->Rewind();
 
 	return true;
@@ -1113,14 +1127,14 @@ bool SMediaSourceOverlay::Command_StepBack(UMediaPlayer* InMediaPlayer)
 
 bool SMediaSourceOverlay::Command_Play(UMediaPlayer* InMediaPlayer)
 {
-	InMediaPlayer->Play();
+	InMediaPlayer->SetRate(1.f);
 
 	return true;
 }
 
 bool SMediaSourceOverlay::Command_Pause(UMediaPlayer* InMediaPlayer)
 {
-	InMediaPlayer->Pause();
+	InMediaPlayer->SetRate(0.f);
 
 	return true;
 }
@@ -1146,7 +1160,7 @@ bool SMediaSourceOverlay::Command_Forward(UMediaPlayer* InMediaPlayer)
 
 	if (!FMath::IsNearlyZero(FrameRate))
 	{
-		InMediaPlayer->Pause();
+		InMediaPlayer->SetRate(0.f);
 
 		const double FrameTime = 1.0 / FrameRate;
 		FTimespan SeekLocation = InMediaPlayer->GetDuration();
diff --git a/Engine/Plugins/Experimental/MediaViewer/Source/MediaViewer/Private/Widgets/SMediaSourceOverlay.h b/Engine/Plugins/Experimental/MediaViewer/Source/MediaViewer/Private/Widgets/SMediaSourceOverlay.h
--- a/Engine/Plugins/Experimental/MediaViewer/Source/MediaViewer/Private/Widgets/SMediaSourceOverlay.h
+++ b/Engine/Plugins/Experimental/MediaViewer/Source/MediaViewer/Private/Widgets/SMediaSourceOverlay.h
@@ -6,6 +6,7 @@
 
 #include "Internationalization/Text.h"
 #include "MediaPlayerEditorModule.h"
+#include "Misc/Optional.h"
 #include "Misc/Timespan.h"
 #include "Misc/TVariant.h"
 
@@ -72,6 +73,7 @@ protected:
 
 	/** When a sync scrub is started, this stores the start of this media. */
 	TVariant<FTimespan, float> SyncedScrubStartTime;
+	TOptional<float> RateWhenScrubEventReceived;
 
 	/** The last time the mouse was moved over the widget. */
 	double LastInteractionTime = -1.0;
