Media Viewer: Fixed controls to properly display when playing in reverse.
- Also added keybind info to tooltip buttons.

diff --git a/Engine/Plugins/Experimental/MediaViewer/Source/MediaViewer/Private/Widgets/SMediaSourceOverlay.cpp b/Engine/Plugins/Experimental/MediaViewer/Source/MediaViewer/Private/Widgets/SMediaSourceOverlay.cpp
--- a/Engine/Plugins/Experimental/MediaViewer/Source/MediaViewer/Private/Widgets/SMediaSourceOverlay.cpp
+++ b/Engine/Plugins/Experimental/MediaViewer/Source/MediaViewer/Private/Widgets/SMediaSourceOverlay.cpp
@@ -52,6 +52,36 @@ namespace MediaSourceOverlay::Events
 	};
 }
 
+namespace MediaSourceOverlay
+{
+	FText GetKeybindText(const TSharedPtr<FUICommandInfo>& InCommand)
+	{
+		if (InCommand.IsValid())
+		{
+			const TSharedRef<const FInputChord>& Chord = InCommand->GetFirstValidChord();
+
+			if (Chord->IsValidChord())
+			{
+				return Chord->GetInputText(/* Long display name */ false);
+			}
+		}
+
+		return FText::GetEmpty();
+	}
+
+	FText AddKeybind(const FText& InToolTip, const TSharedPtr<FUICommandInfo>& InCommand)
+	{
+		const FText KeybindText = GetKeybindText(InCommand);
+
+		if (!KeybindText.IsEmpty())
+		{
+			return FText::Format(LOCTEXT("AddKeybind", "{0}\nKeybind: {1}"), InToolTip, KeybindText);
+		}
+
+		return InToolTip;
+	}
+}
+
 void SMediaSourceOverlay::PrivateRegisterAttributes(FSlateAttributeDescriptor::FInitializer&)
 {
 }
@@ -228,6 +258,106 @@ TSharedRef<SWidget> SMediaSourceOverlay::CreateSlider()
 
 TSharedRef<SWidget> SMediaSourceOverlay::CreateControls()
 {
+	const FMediaViewerCommands& MediaViewerCommands = FMediaViewerCommands::Get();
+	
+	const FText RelatedKeybindsText = LOCTEXT("RelatedKeybinds", "Related Keybinds:");
+
+	TArray<FText> BackwardsKeybinds;
+	BackwardsKeybinds.Reserve(4);
+	{
+		const FText StepBackKeybind = MediaSourceOverlay::GetKeybindText(MediaViewerCommands.StepBackVideo);
+		const FText RewindKeybind = MediaSourceOverlay::GetKeybindText(MediaViewerCommands.RewindVideo);
+		const FText QuickRewindKeybind = MediaSourceOverlay::GetKeybindText(MediaViewerCommands.QuickRewindVideo);
+
+		if (!StepBackKeybind.IsEmpty() || !RewindKeybind.IsEmpty() || !QuickRewindKeybind.IsEmpty())
+		{
+			BackwardsKeybinds.Add(RelatedKeybindsText);
+
+			if (!StepBackKeybind.IsEmpty())
+			{
+				BackwardsKeybinds.Add(FText::Format(
+					LOCTEXT("StepBackKeybind", "- Step Back: {0}"),
+					StepBackKeybind
+				));
+			}
+
+			if (!RewindKeybind.IsEmpty())
+			{
+				BackwardsKeybinds.Add(FText::Format(
+					LOCTEXT("RewindKeybind", "- Rewind 1s: {0}"),
+					RewindKeybind
+				));
+			}
+
+			if (!QuickRewindKeybind.IsEmpty())
+			{
+				BackwardsKeybinds.Add(FText::Format(
+					LOCTEXT("QuickRewindKeybind", "- Rewind 10s: {0}"),
+					QuickRewindKeybind
+				));
+			}
+		}
+	}
+
+	FText StepBackTooltip = LOCTEXT("StepBackward", "Step backward 1 frame.\n\nOnly available while paused.");
+
+	if (!BackwardsKeybinds.IsEmpty())
+	{
+		StepBackTooltip = FText::Format(
+			INVTEXT("{0}\n\n{1}"),
+			StepBackTooltip,
+			FText::Join(INVTEXT("\n"), BackwardsKeybinds)
+		);
+	}
+
+	TArray<FText> ForwardsKeybinds;
+	ForwardsKeybinds.Reserve(4);
+	{
+		const FText StepForwardKeybind = MediaSourceOverlay::GetKeybindText(MediaViewerCommands.StepForwardVideo);
+		const FText FastForwardKeybind = MediaSourceOverlay::GetKeybindText(MediaViewerCommands.FastForwardVideo);
+		const FText QuickFastForwardKeybind = MediaSourceOverlay::GetKeybindText(MediaViewerCommands.QuickFastForwardVideo);
+
+		if (!StepForwardKeybind.IsEmpty() || !FastForwardKeybind.IsEmpty() || !QuickFastForwardKeybind.IsEmpty())
+		{
+			ForwardsKeybinds.Add(RelatedKeybindsText);
+
+			if (!StepForwardKeybind.IsEmpty())
+			{
+				ForwardsKeybinds.Add(FText::Format(
+					LOCTEXT("StepForwardKeybind", "- Step Forward: {0}"),
+					StepForwardKeybind
+				));
+			}
+
+			if (!FastForwardKeybind.IsEmpty())
+			{
+				ForwardsKeybinds.Add(FText::Format(
+					LOCTEXT("FastForwardKeybind", "- Fast Forward 1s: {0}"),
+					FastForwardKeybind
+				));
+			}
+
+			if (!QuickFastForwardKeybind.IsEmpty())
+			{
+				ForwardsKeybinds.Add(FText::Format(
+					LOCTEXT("QuickFastForwardKeybind", "- Fast Forward 10s: {0}"),
+					QuickFastForwardKeybind
+				));
+			}
+		}
+	}
+
+	FText StepForwardTooltip = LOCTEXT("StepForward", "Step forward 1 frame.\n\nOnly available while paused.");
+
+	if (!ForwardsKeybinds.IsEmpty())
+	{
+		StepForwardTooltip = FText::Format(
+			INVTEXT("{0}\n\n{1}"),
+			StepForwardTooltip,
+			FText::Join(INVTEXT("\n"), ForwardsKeybinds)
+		);
+	}
+
 	return SNew(SHorizontalBox)
 
 		// Current frame
@@ -291,7 +421,7 @@ TSharedRef<SWidget> SMediaSourceOverlay::CreateControls()
 				.DesiredSizeOverride(FVector2D(20.f))
 			]
 		]
-
+		
 		// Step back button.
 		+ SHorizontalBox::Slot()
 		.VAlign(VAlign_Center)
@@ -307,7 +437,7 @@ TSharedRef<SWidget> SMediaSourceOverlay::CreateControls()
 				SNew(SImage)
 				.ColorAndOpacity(FSlateColor::UseForeground())
 				.Image(FSlateIcon(FAppStyle::GetAppStyleSetName(), "Animation.Backward_Step").GetIcon())
-				.ToolTipText(LOCTEXT("StepBack", "Step back 1 frame.\n\nOnly available while paused."))
+				.ToolTipText(StepBackTooltip)
 				.DesiredSizeOverride(FVector2D(20.f))
 			]
 		]
@@ -367,7 +497,7 @@ TSharedRef<SWidget> SMediaSourceOverlay::CreateControls()
 				SNew(SImage)
 				.ColorAndOpacity(FSlateColor::UseForeground())
 				.Image(FSlateIcon(FAppStyle::GetAppStyleSetName(), "Animation.Forward_Step").GetIcon())
-				.ToolTipText(LOCTEXT("StepForward", "Step forward 1 frame.\n\nOnly available while paused."))
+				.ToolTipText(StepForwardTooltip)
 				.DesiredSizeOverride(FVector2D(20.f))
 			]
 		]
@@ -562,7 +692,7 @@ const FSlateBrush* SMediaSourceOverlay::Reverse_GetBrush() const
 {
 	if (UMediaPlayer* MediaPlayer = GetMediaPlayer())
 	{
-		const float Rate = FMath::Abs(MediaPlayer->GetRate());
+		const float Rate = MediaPlayer->GetRate();
 
 		if (Rate < 0 && !FMath::IsNearlyZero(Rate))
 		{
@@ -577,7 +707,7 @@ FText SMediaSourceOverlay::Reverse_GetToolTip() const
 {
 	if (UMediaPlayer* MediaPlayer = GetMediaPlayer())
 	{
-		const float Rate = FMath::Abs(MediaPlayer->GetRate());
+		const float Rate = MediaPlayer->GetRate();
 
 		if (Rate < 0 && !FMath::IsNearlyZero(Rate))
 		{
@@ -592,9 +722,9 @@ FReply SMediaSourceOverlay::Reverse_OnClicked()
 {
 	if (UMediaPlayer* MediaPlayer = GetMediaPlayer())
 	{
-		const float Rate = FMath::Abs(MediaPlayer->GetRate());
+		const float Rate = MediaPlayer->GetRate();
 
-		if (FMath::IsNearlyZero(Rate))
+		if (Rate > 0 || FMath::IsNearlyZero(Rate))
 		{
 			if (Command_Reverse(MediaPlayer))
 			{
@@ -617,8 +747,6 @@ bool SMediaSourceOverlay::StepBack_IsEnabled() const
 {
 	if (UMediaPlayer* MediaPlayer = GetMediaPlayer())
 	{
-		const float Rate = FMath::Abs(MediaPlayer->GetRate());
-
 		return MediaPlayer->IsReady() && MediaPlayer->IsPaused();
 	}
 
@@ -652,7 +780,7 @@ const FSlateBrush* SMediaSourceOverlay::Play_GetBrush() const
 {
 	if (UMediaPlayer* MediaPlayer = GetMediaPlayer())
 	{
-		const float Rate = FMath::Abs(MediaPlayer->GetRate());
+		const float Rate = MediaPlayer->GetRate();
 
 		if (Rate > 0 && !FMath::IsNearlyZero(Rate))
 		{
@@ -665,26 +793,34 @@ const FSlateBrush* SMediaSourceOverlay::Play_GetBrush() const
 
 FText SMediaSourceOverlay::Play_GetToolTip() const
 {
+	const FMediaViewerCommands& MediaViewerCommands = FMediaViewerCommands::Get();
+
 	if (UMediaPlayer* MediaPlayer = GetMediaPlayer())
 	{
-		const float Rate = FMath::Abs(MediaPlayer->GetRate());
+		const float Rate = MediaPlayer->GetRate();
 
 		if (Rate > 0 && !FMath::IsNearlyZero(Rate))
 		{
-			return LOCTEXT("Pause", "Pause media playback");
+			return MediaSourceOverlay::AddKeybind(
+				LOCTEXT("Pause", "Pause media playback.\n"),
+				MediaViewerCommands.ToggleVideoPlay
+			);
 		}
 	}
 
-	return LOCTEXT("Play", "Play media forward");
+	return MediaSourceOverlay::AddKeybind(
+		LOCTEXT("Play", "Play media forward.\n"),
+		MediaViewerCommands.ToggleVideoPlay
+	);
 }
 
 FReply SMediaSourceOverlay::Play_OnClicked()
 {
 	if (UMediaPlayer* MediaPlayer = GetMediaPlayer())
 	{
-		const float Rate = FMath::Abs(MediaPlayer->GetRate());
+		const float Rate = MediaPlayer->GetRate();
 
-		if (FMath::IsNearlyZero(Rate))
+		if (Rate < 0 || FMath::IsNearlyZero(Rate))
 		{
 			if (Command_Play(MediaPlayer))
 			{
@@ -707,8 +843,6 @@ bool SMediaSourceOverlay::StepForward_IsEnabled() const
 {
 	if (UMediaPlayer* MediaPlayer = GetMediaPlayer())
 	{
-		const float Rate = FMath::Abs(MediaPlayer->GetRate());
-
 		return MediaPlayer->IsReady() && MediaPlayer->IsPaused();
 	}
 
diff --git a/Engine/Plugins/Experimental/MediaViewer/Source/MediaViewer/Public/MediaViewerDelegates.h b/Engine/Plugins/Experimental/MediaViewer/Source/MediaViewer/Public/MediaViewerDelegates.h
--- a/Engine/Plugins/Experimental/MediaViewer/Source/MediaViewer/Public/MediaViewerDelegates.h
+++ b/Engine/Plugins/Experimental/MediaViewer/Source/MediaViewer/Public/MediaViewerDelegates.h
@@ -68,7 +68,7 @@ struct FMediaViewerDelegates
 	/** Sets the orientation in A/B mode. Will set the next expected orientation even if not in A/B mode. */
 	FSetOrientation SetABOrientation;
 
-	/** Returns a bitflag value indicating which image viewers are currently active (have images, not null or null image viewer. */
+	/** Returns a bitflag value indicating which image viewers are currently active (have images, not null or null image viewer). */
 	FGetActiveView GetActiveView;
 
 	/* Gets a pointer to the Media Viewer configurable settings. */
