Media Viewer: Button tooltip redux

diff --git a/Engine/Plugins/Experimental/MediaViewer/Source/MediaViewer/Private/Widgets/SMediaSourceOverlay.cpp b/Engine/Plugins/Experimental/MediaViewer/Source/MediaViewer/Private/Widgets/SMediaSourceOverlay.cpp
--- a/Engine/Plugins/Experimental/MediaViewer/Source/MediaViewer/Private/Widgets/SMediaSourceOverlay.cpp
+++ b/Engine/Plugins/Experimental/MediaViewer/Source/MediaViewer/Private/Widgets/SMediaSourceOverlay.cpp
@@ -711,11 +711,11 @@ FText SMediaSourceOverlay::Reverse_GetToolTip() const
 
 		if (Rate < 0 && !FMath::IsNearlyZero(Rate))
 		{
-			return LOCTEXT("Pause", "Pause media playback");
+			return LOCTEXT("Reverse_Pause", "Pause media playback");
 		}
 	}
 
-	return LOCTEXT("Reverse", "Play media in reverse.\n\nNot widely supported by media decoders.");
+	return LOCTEXT("Reverse_Reverse", "Play media in reverse.\n\nNot widely supported by media decoders.");
 }
 
 FReply SMediaSourceOverlay::Reverse_OnClicked()
@@ -802,14 +802,14 @@ FText SMediaSourceOverlay::Play_GetToolTip() const
 		if (Rate > 0 && !FMath::IsNearlyZero(Rate))
 		{
 			return MediaSourceOverlay::AddKeybind(
-				LOCTEXT("Pause", "Pause media playback.\n"),
+				LOCTEXT("Play_Pause", "Pause media playback.\n"),
 				MediaViewerCommands.ToggleVideoPlay
 			);
 		}
 	}
 
 	return MediaSourceOverlay::AddKeybind(
-		LOCTEXT("Play", "Play media forward.\n"),
+		LOCTEXT("Play_Play", "Play media forward.\n"),
 		MediaViewerCommands.ToggleVideoPlay
 	);
 }
