Media Viewer: Added mouse capture

diff --git a/Engine/Plugins/Experimental/MediaViewer/Source/MediaViewer/Private/Widgets/SMediaImageViewerOverlay.cpp b/Engine/Plugins/Experimental/MediaViewer/Source/MediaViewer/Private/Widgets/SMediaImageViewerOverlay.cpp
--- a/Engine/Plugins/Experimental/MediaViewer/Source/MediaViewer/Private/Widgets/SMediaImageViewerOverlay.cpp
+++ b/Engine/Plugins/Experimental/MediaViewer/Source/MediaViewer/Private/Widgets/SMediaImageViewerOverlay.cpp
@@ -204,7 +204,7 @@ FReply SMediaImageViewerOverlay::OnMouseButtonDown(const FGeometry& InMyGeometry
 			DragStartOffset = GetOffset();
 			DragButtonName = InMouseEvent.GetEffectingButton().GetFName();
 			DragType = EDragType::Offset;
-			return FReply::Handled();
+			return FReply::Handled().CaptureMouse(SharedThis(this));
 		}
 		else if (InMouseEvent.GetEffectingButton() == EKeys::RightMouseButton && InMouseEvent.IsAltDown())
 		{
@@ -212,7 +212,7 @@ FReply SMediaImageViewerOverlay::OnMouseButtonDown(const FGeometry& InMyGeometry
 			LastDragScaleValue = GetScale();
 			DragButtonName = InMouseEvent.GetEffectingButton().GetFName();
 			DragType = EDragType::Scale;
-			return FReply::Handled();
+			return FReply::Handled().CaptureMouse(SharedThis(this));
 		}
 
 		if (TSharedPtr<FMediaImageViewer> ImageViewer = Delegates->GetImageViewer.Execute(Position))
@@ -223,7 +223,7 @@ FReply SMediaImageViewerOverlay::OnMouseButtonDown(const FGeometry& InMyGeometry
 			{
 				DragButtonName = InMouseEvent.GetEffectingButton().GetFName();
 				DragType = EDragType::External;
-				return FReply::Handled();
+				return FReply::Handled().CaptureMouse(SharedThis(this));
 			}
 		}
 	}
@@ -251,6 +251,7 @@ FReply SMediaImageViewerOverlay::OnMouseButtonUp(const FGeometry& InMyGeometry,
 	else
 	{
 		OnDragButtonUp(InMouseEvent.GetEffectingButton().GetFName());
+		return SCompoundWidget::OnMouseButtonUp(InMyGeometry, InMouseEvent).ReleaseMouseCapture();
 	}
 
 	return SCompoundWidget::OnMouseButtonUp(InMyGeometry, InMouseEvent);
