Media Viewer: Support for multiple "standalone" media viewer tabs all using a single library, which is now owned by the module.

diff --git a/Engine/Plugins/Experimental/MediaViewer/Source/MediaViewer/Private/Library/MediaViewerLibrary.cpp b/Engine/Plugins/Experimental/MediaViewer/Source/MediaViewer/Private/Library/MediaViewerLibrary.cpp
--- a/Engine/Plugins/Experimental/MediaViewer/Source/MediaViewer/Private/Library/MediaViewerLibrary.cpp
+++ b/Engine/Plugins/Experimental/MediaViewer/Source/MediaViewer/Private/Library/MediaViewerLibrary.cpp
@@ -35,20 +35,12 @@ const FGuid& FMediaViewerLibrary::GetSnapshotsGroupId() const
 	return SnapshotGroupId;
 }
 
-FMediaViewerLibrary::FMediaViewerLibrary(const FArgs& InArgs)
+FMediaViewerLibrary::FMediaViewerLibrary()
 {
-	const FText DefaultTitle = InArgs.DefaultGroupName.IsEmpty()
-		? LOCTEXT("DefaultGroup", "Pinned")
-		: InArgs.DefaultGroupName;
-
-	const FText DefaultToolTip = InArgs.DefaultGroupToolTip.IsEmpty()
-		? LOCTEXT("DefaultGroupTooltip", "Pinned items.")
-		: InArgs.DefaultGroupToolTip;
-
 	Groups.Add(MakeShared<FMediaViewerLibraryGroup>(
 		GetDefaultGroupId(),
-		DefaultTitle,
-		DefaultToolTip,
+		LOCTEXT("DefaultGroup", "Pinned"),
+		LOCTEXT("DefaultGroupTooltip", "Pinned items."),
 		/* Dynamic */ false
 	));
 
@@ -107,7 +99,7 @@ bool FMediaViewerLibrary::AddGroup(const TSharedRef<FMediaViewerLibraryGroup>& I
 		);
 	}
 
-	OnChanged(IMediaViewerLibrary::EChangeType::GroupAdded);
+	OnChanged(FMediaViewerLibrary::EChangeType::GroupAdded);
 
 	return true;
 }
@@ -136,7 +128,7 @@ int32 FMediaViewerLibrary::EmptyGroup(const FGuid& InGroupId)
 
 	const int32 Count = Group->Empty();	
 
-	OnChanged(IMediaViewerLibrary::EChangeType::ItemRemoved);
+	OnChanged(FMediaViewerLibrary::EChangeType::ItemRemoved);
 
 	return Count;
 }
@@ -185,7 +177,7 @@ TSharedPtr<FMediaViewerLibraryGroup> FMediaViewerLibrary::RemoveGroup(const FGui
 
 	Groups.RemoveAt(GroupToRemoveIndex);
 
-	OnChanged(IMediaViewerLibrary::EChangeType::GroupRemoved);
+	OnChanged(FMediaViewerLibrary::EChangeType::GroupRemoved);
 
 	return GroupToRemove;
 }
@@ -251,7 +243,7 @@ bool FMediaViewerLibrary::AddItemToGroup(const TSharedRef<FMediaViewerLibraryIte
 {
 	const FGuid& ItemId = InNewItem->GetId();
 
-	if (GetItemGroup(InNewItem->GetId()).IsValid())
+	if (InTargetGroupId != HistoryGroupId && GetItemGroup(InNewItem->GetId()).IsValid())
 	{
 		return false;
 	}
@@ -272,7 +264,7 @@ bool FMediaViewerLibrary::AddItemToGroup(const TSharedRef<FMediaViewerLibraryIte
 
 	TargetGroup->AddItem(ItemId, InIndex);
 
-	OnChanged(IMediaViewerLibrary::EChangeType::ItemGroupChanged);
+	OnChanged(FMediaViewerLibrary::EChangeType::ItemGroupChanged);
 
 	return true;
 }
@@ -309,7 +301,7 @@ bool FMediaViewerLibrary::AddItemBelowItem(const TSharedRef<FMediaViewerLibraryI
 
 	TargetGroup->AddItem(ItemId, TargetIndex + 1);
 
-	OnChanged(IMediaViewerLibrary::EChangeType::ItemGroupChanged);
+	OnChanged(FMediaViewerLibrary::EChangeType::ItemGroupChanged);
 
 	return true;
 }
@@ -340,7 +332,7 @@ bool FMediaViewerLibrary::MoveItemToGroup(const FGroupItem& InItemToMove, const
 
 	TargetGroup->AddItem(InItemToMove.ItemId, InIndex);
 
-	OnChanged(IMediaViewerLibrary::EChangeType::ItemGroupChanged);
+	OnChanged(FMediaViewerLibrary::EChangeType::ItemGroupChanged);
 
 	return true;
 }
@@ -375,7 +367,7 @@ bool FMediaViewerLibrary::MoveItemWithinGroup(const FGroupItem& InItemToMove, in
 	Group->RemoveItem(InItemToMove.ItemId);
 	Group->AddItem(InItemToMove.ItemId, InIndex);
 
-	OnChanged(IMediaViewerLibrary::EChangeType::ItemGroupChanged);
+	OnChanged(FMediaViewerLibrary::EChangeType::ItemGroupChanged);
 
 	return true;
 }
@@ -417,7 +409,7 @@ bool FMediaViewerLibrary::MoveItemBelowItem(const FGroupItem& InItemToMove, cons
 
 	TargetGroup->AddItem(InItemToMove.ItemId, TargetIndex + 1);
 
-	OnChanged(IMediaViewerLibrary::EChangeType::ItemGroupChanged);
+	OnChanged(FMediaViewerLibrary::EChangeType::ItemGroupChanged);
 
 	return true;
 }
@@ -438,10 +430,12 @@ bool FMediaViewerLibrary::RemoveItemFromGroup(const FGroupItem& InItemToRemove)
 		return false;
 	}
 
-	Group->RemoveItem(InItemToRemove.ItemId);
-
-	OnChanged(IMediaViewerLibrary::EChangeType::ItemGroupChanged);
+	if (!Group->RemoveItem(InItemToRemove.ItemId))
+	{
+		return false;
+	}
 
+	OnChanged(FMediaViewerLibrary::EChangeType::ItemGroupChanged);
 	return true;
 }
 
@@ -491,7 +485,7 @@ TSharedPtr<FMediaViewerLibraryItem> FMediaViewerLibrary::RemoveItem(const FGuid&
 
 	Items.Remove(InItemIdToRemove);
 
-	OnChanged(IMediaViewerLibrary::EChangeType::ItemRemoved);
+	OnChanged(FMediaViewerLibrary::EChangeType::ItemRemoved);
 
 	return ItemToRemove;
 }
@@ -523,9 +517,9 @@ void FMediaViewerLibrary::RemoveInvalidGroupItems(const FGuid& InGroup)
 	}
 }
 
-void FMediaViewerLibrary::OnChanged(IMediaViewerLibrary::EChangeType InChangeType)
+void FMediaViewerLibrary::OnChanged(FMediaViewerLibrary::EChangeType InChangeType)
 {
-	OnChangedDelegate.ExecuteIfBound(SharedThis(this), InChangeType);
+	OnChangedDelegate.Broadcast(SharedThis(this), InChangeType);
 }
 
 } // UE::MediaViewer::Private
diff --git a/Engine/Plugins/Experimental/MediaViewer/Source/MediaViewer/Private/Library/MediaViewerLibrary.h b/Engine/Plugins/Experimental/MediaViewer/Source/MediaViewer/Private/Library/MediaViewerLibrary.h
--- a/Engine/Plugins/Experimental/MediaViewer/Source/MediaViewer/Private/Library/MediaViewerLibrary.h
+++ b/Engine/Plugins/Experimental/MediaViewer/Source/MediaViewer/Private/Library/MediaViewerLibrary.h
@@ -18,14 +18,14 @@ namespace UE::MediaViewer::Private
 constexpr int32 MaxHistoryEntries = 20;
 
 /**
- * Implementation of @see IMediaViewerLibrary.
+ * Implementation of @see FMediaViewerLibrary.
  */
 class FMediaViewerLibrary : public IMediaViewerLibrary, public TSharedFromThis<FMediaViewerLibrary>
 {
 	friend class ::UMediaViewerLibraryIni;
 
 public:
-	FMediaViewerLibrary(const FArgs& InArgs);
+	FMediaViewerLibrary();
 
 	virtual ~FMediaViewerLibrary() = default;
 
@@ -33,7 +33,7 @@ public:
 
 	bool CanDragDropItem(const FGroupItem& InItem) const;
 
-	//~ Begin IMediaViewerLibrary
+	//~ Begin FMediaViewerLibrary
 	virtual const FGuid& GetDefaultGroupId() const override;
 	virtual const FGuid& GetHistoryGroupId() const override;
 	virtual const FGuid& GetSnapshotsGroupId() const override;
@@ -59,14 +59,14 @@ public:
 	virtual bool CanRemoveItem(const FGuid& InItemIdToRemove) const override;
 	virtual TSharedPtr<FMediaViewerLibraryItem> RemoveItem(const FGuid& InItemIdToRemove) override;
 	virtual FOnChanged::RegistrationType& GetOnChanged() override;
-	//~ End IMediaViewerLibrary
+	//~ End FMediaViewerLibrary
 
 protected:
 	TArray<TSharedRef<FMediaViewerLibraryGroup>> Groups;
 	TMap<FGuid, TSharedRef<FMediaViewerLibraryItem>> Items;
 	FOnChanged OnChangedDelegate;
 
-	void OnChanged(IMediaViewerLibrary::EChangeType InChangeType);
+	void OnChanged(FMediaViewerLibrary::EChangeType InChangeType);
 };
 
 } // UE::MediaViewer::Private
\ No newline at end of file
diff --git a/Engine/Plugins/Experimental/MediaViewer/Source/MediaViewer/Private/Library/MediaViewerLibraryIni.cpp b/Engine/Plugins/Experimental/MediaViewer/Source/MediaViewer/Private/Library/MediaViewerLibraryIni.cpp
--- a/Engine/Plugins/Experimental/MediaViewer/Source/MediaViewer/Private/Library/MediaViewerLibraryIni.cpp
+++ b/Engine/Plugins/Experimental/MediaViewer/Source/MediaViewer/Private/Library/MediaViewerLibraryIni.cpp
@@ -15,12 +15,16 @@ UMediaViewerLibraryIni& UMediaViewerLibraryIni::Get()
 	return *GetMutableDefault<UMediaViewerLibraryIni>();
 }
 
-void UMediaViewerLibraryIni::SaveLibrary(const TSharedRef<UE::MediaViewer::Private::FMediaViewerLibrary>& InLibrary)
+void UMediaViewerLibraryIni::SaveLibrary(const TSharedRef<UE::MediaViewer::IMediaViewerLibrary>& InLibrary)
 {
+	using namespace UE::MediaViewer::Private;
+
+	TSharedRef<FMediaViewerLibrary> LibraryImpl = StaticCastSharedRef<FMediaViewerLibrary>(InLibrary);
+
 	Groups.Empty();
-	Groups.Reserve(InLibrary->Groups.Num());
+	Groups.Reserve(LibraryImpl->Groups.Num());
 
-	for (const TSharedRef<const FMediaViewerLibraryGroup>& Group : InLibrary->GetGroups())
+	for (const TSharedRef<FMediaViewerLibraryGroup>& Group : LibraryImpl->Groups)
 	{
 		if (Group->IsDynamic())
 		{
@@ -31,9 +35,9 @@ void UMediaViewerLibraryIni::SaveLibrary(const TSharedRef<UE::MediaViewer::Priva
 	}
 
 	Items.Empty();
-	Items.Reserve(InLibrary->Items.Num());
+	Items.Reserve(LibraryImpl->Items.Num());
 
-	for (const TPair<FGuid, TSharedRef<FMediaViewerLibraryItem>>& ItemPair : InLibrary->Items)
+	for (const TPair<FGuid, TSharedRef<FMediaViewerLibraryItem>>& ItemPair : LibraryImpl->Items)
 	{
 		const TSharedRef<FMediaViewerLibraryItem>& Item = ItemPair.Value;
 
@@ -46,31 +50,35 @@ void UMediaViewerLibraryIni::SaveLibrary(const TSharedRef<UE::MediaViewer::Priva
 	}
 }
 
-void UMediaViewerLibraryIni::LoadLibrary(const TSharedRef<UE::MediaViewer::Private::FMediaViewerLibrary>& InLibrary) const
+void UMediaViewerLibraryIni::LoadLibrary(const TSharedRef<UE::MediaViewer::IMediaViewerLibrary>& InLibrary) const
 {
+	using namespace UE::MediaViewer::Private;
+
+	TSharedRef<FMediaViewerLibrary> LibraryImpl = StaticCastSharedRef<FMediaViewerLibrary>(InLibrary);
+
 	UE::MediaViewer::IMediaViewerModule& Module = UE::MediaViewer::IMediaViewerModule::Get();
 
 	for (const FMediaViewerLibraryItemData& ItemData : Items)
 	{
 		const FGuid ItemId = ItemData.Item.GetId();
 
-		if (!InLibrary->Items.Contains(ItemId))
+		if (!LibraryImpl->Items.Contains(ItemId))
 		{
 			if (TSharedPtr<FMediaViewerLibraryItem> Item = Module.CreateLibraryItem(ItemData.ItemType, ItemData.Item))
 			{
-				InLibrary->Items.Add(ItemId, Item.ToSharedRef());
+				LibraryImpl->Items.Add(ItemId, Item.ToSharedRef());
 			}
 		}
 	}
 
 	for (const FMediaViewerLibraryGroup& SavedGroup : Groups)
 	{
-		TSharedPtr<FMediaViewerLibraryGroup> Group = InLibrary->GetGroup(SavedGroup.GetId());
+		TSharedPtr<FMediaViewerLibraryGroup> Group = LibraryImpl->GetGroup(SavedGroup.GetId());
 
 		if (!Group.IsValid())
 		{
 			Group = MakeShared<FMediaViewerLibraryGroup>(FMediaViewerLibraryGroup::FPrivateToken(), SavedGroup);
-			InLibrary->AddGroup(Group.ToSharedRef());
+			LibraryImpl->AddGroup(Group.ToSharedRef());
 		}
 
 		TSet<FGuid> ExistingItems;
@@ -85,7 +93,7 @@ void UMediaViewerLibraryIni::LoadLibrary(const TSharedRef<UE::MediaViewer::Priva
 			}
 		}
 
-		InLibrary->RemoveInvalidGroupItems(SavedGroup.GetId());
+		LibraryImpl->RemoveInvalidGroupItems(SavedGroup.GetId());
 	}
 }
 
diff --git a/Engine/Plugins/Experimental/MediaViewer/Source/MediaViewer/Private/Library/MediaViewerLibraryIni.h b/Engine/Plugins/Experimental/MediaViewer/Source/MediaViewer/Private/Library/MediaViewerLibraryIni.h
--- a/Engine/Plugins/Experimental/MediaViewer/Source/MediaViewer/Private/Library/MediaViewerLibraryIni.h
+++ b/Engine/Plugins/Experimental/MediaViewer/Source/MediaViewer/Private/Library/MediaViewerLibraryIni.h
@@ -13,9 +13,9 @@
 
 #include "MediaViewerLibraryIni.generated.h"
 
-namespace UE::MediaViewer::Private
+namespace UE::MediaViewer
 {
-	class FMediaViewerLibrary;
+	class IMediaViewerLibrary;
 }
 
 USTRUCT()
@@ -71,9 +71,9 @@ class UMediaViewerLibraryIni : public UObject
 public:
 	static UMediaViewerLibraryIni& Get();
 
-	void SaveLibrary(const TSharedRef<UE::MediaViewer::Private::FMediaViewerLibrary>& InLibrary);
+	void SaveLibrary(const TSharedRef<UE::MediaViewer::IMediaViewerLibrary>& InLibrary);
 
-	void LoadLibrary(const TSharedRef<UE::MediaViewer::Private::FMediaViewerLibrary>& InLibrary) const;
+	void LoadLibrary(const TSharedRef<UE::MediaViewer::IMediaViewerLibrary>& InLibrary) const;
 
 	bool HasGroup(const FGuid& InGroupId) const;
 
diff --git a/Engine/Plugins/Experimental/MediaViewer/Source/MediaViewer/Private/MediaViewerCommands.cpp b/Engine/Plugins/Experimental/MediaViewer/Source/MediaViewer/Private/MediaViewerCommands.cpp
--- a/Engine/Plugins/Experimental/MediaViewer/Source/MediaViewer/Private/MediaViewerCommands.cpp
+++ b/Engine/Plugins/Experimental/MediaViewer/Source/MediaViewer/Private/MediaViewerCommands.cpp
@@ -347,6 +347,15 @@ void FMediaViewerCommands::RegisterCommands()
 		EUserInterfaceActionType::Button,
 		FInputChord(EKeys::L, EModifierKey::Shift)
 	);
+
+	UI_COMMAND(
+		OpenNewTab,
+		"New Tab",
+		"Open a new standalone tab.",
+		EUserInterfaceActionType::Button,
+		FInputChord(EKeys::N, EModifierKey::Control),
+		FInputChord(EKeys::M, EModifierKey::Command)
+	);
 }
 
 } // UE::MediaViewer::Private
diff --git a/Engine/Plugins/Experimental/MediaViewer/Source/MediaViewer/Private/MediaViewerCommands.h b/Engine/Plugins/Experimental/MediaViewer/Source/MediaViewer/Private/MediaViewerCommands.h
--- a/Engine/Plugins/Experimental/MediaViewer/Source/MediaViewer/Private/MediaViewerCommands.h
+++ b/Engine/Plugins/Experimental/MediaViewer/Source/MediaViewer/Private/MediaViewerCommands.h
@@ -68,6 +68,8 @@ public:
 	TSharedPtr<FUICommandInfo> StepForwardVideo;
 	TSharedPtr<FUICommandInfo> FastForwardVideo;
 	TSharedPtr<FUICommandInfo> QuickFastForwardVideo;
+
+	TSharedPtr<FUICommandInfo> OpenNewTab;
 };
 
 } // UE::MediaViewer::Private
diff --git a/Engine/Plugins/Experimental/MediaViewer/Source/MediaViewer/Private/MediaViewerModule.cpp b/Engine/Plugins/Experimental/MediaViewer/Source/MediaViewer/Private/MediaViewerModule.cpp
--- a/Engine/Plugins/Experimental/MediaViewer/Source/MediaViewer/Private/MediaViewerModule.cpp
+++ b/Engine/Plugins/Experimental/MediaViewer/Source/MediaViewer/Private/MediaViewerModule.cpp
@@ -12,6 +12,9 @@
 #include "ImageViewers/Texture2DImageViewer.h"
 #include "ImageViewers/TextureRenderTarget2DImageViewer.h"
 #include "LevelEditor.h"
+#include "Library/LevelEditorViewportGroup.h"
+#include "Library/MediaTextureGroup.h"
+#include "Library/MediaTrackGroup.h"
 #include "Library/MediaViewerLibrary.h"
 #include "Library/MediaViewerLibraryIni.h"
 #include "MediaViewer.h"
@@ -49,6 +52,26 @@ TSharedPtr<FTabManager> GetLevelEditorTabManager()
 }
 
 const FLazyName FMediaViewerModule::TabId = "MediaViewerTabId";
+const FLazyName FMediaViewerModule::StandaloneTabId = "MediaViewerStandaloneTabId";
+
+FMediaViewerModule::FMediaViewerModule()
+	: Library(MakeShared<FMediaViewerLibrary>())
+{
+	static FDelayedAutoRegisterHelper LoadLibraryCallback(
+		EDelayedRegisterRunPhase::EndOfEngineInit,
+		[]
+		{
+			TSharedRef<FMediaViewerLibrary> LibraryImpl = StaticCastSharedRef<FMediaViewerLibrary>(IMediaViewerModule::Get().GetLibrary());
+
+			UMediaViewerLibraryIni::Get().LoadLibrary(LibraryImpl);
+
+			// Add dynamic groups
+			LibraryImpl->AddGroup(MakeShared<FLevelEditorViewportGroup>(LibraryImpl));
+			LibraryImpl->AddGroup(MakeShared<FMediaTrackGroup>(LibraryImpl));
+			LibraryImpl->AddGroup(MakeShared<FMediaTextureGroup>(LibraryImpl));
+		}
+	);
+}
 
 bool FMediaViewerModule::IsFactoryRegistered(FName InFactoryName) const
 {
@@ -153,14 +176,9 @@ bool FMediaViewerModule::OpenTab(const FMediaViewerArgs& InMediaViewerArgs)
 	return true;
 }
 
-TSharedPtr<IMediaViewerLibrary> FMediaViewerModule::GetLibrary() const
+TSharedRef<IMediaViewerLibrary> FMediaViewerModule::GetLibrary() const
 {
-	if (TSharedPtr<SMediaViewerTab> ViewerTab = GetLiveTab())
-	{
-		return ViewerTab->GetLibrary();
-	}
-
-	return nullptr;
+	return Library;
 }
 
 TSharedPtr<FMediaViewerLibraryItem> FMediaViewerModule::CreateLibraryItem(const FAssetData& InAssetData) const
@@ -212,6 +230,35 @@ TSharedRef<IMediaViewerLibraryWidget> FMediaViewerModule::CreateLibrary(const IM
 	return SNew(SMediaViewerLibrary, InArgs, MakeShared<FMediaViewerDelegates>());
 }
 
+TSharedPtr<SDockTab> FMediaViewerModule::OpenStandaloneTab()
+{
+	TSharedPtr<FTabManager> TabManager = UE::MediaViewer::Private::GetLevelEditorTabManager();
+
+	if (!TabManager.IsValid())
+	{
+		return nullptr;
+	}
+
+	TSharedRef<SMediaViewerTab> ViewerTab = CreateMediaViewer({});
+	
+	TSharedRef<SDockTab> StandaloneTab = SNew(SDockTab)
+		.Label(LOCTEXT("MediaViewerTitleStandalone", "Media Viewer (Standalone)"))
+		.ContentPadding(3.f)
+		[
+			// Create with default settings
+			ViewerTab
+		];
+
+	TabManager->InsertNewDocumentTab(
+		TabId,
+		StandaloneTabId,
+		FTabManager::FLiveTabSearch(),
+		StandaloneTab
+	);
+
+	return StandaloneTab;
+}
+
 bool FMediaViewerModule::HasImage(EMediaImageViewerPosition InPosition) const
 {
 	TSharedPtr<FMediaImageViewer> Viewer = GetImage(InPosition);
@@ -358,16 +405,10 @@ void FMediaViewerModule::OnTabClosed(TSharedRef<SDockTab> InDockTab)
 	if (Content->GetWidgetClass().GetWidgetType() == SMediaViewerTab::StaticWidgetClass().GetWidgetType())
 	{
 		TSharedRef<SMediaViewerTab> MediaViewerTab = StaticCastSharedRef<SMediaViewerTab>(Content);
-		MediaViewerTab->GetViewer()->SaveStates();
+		MediaViewerTab->GetViewer()->SaveStates(/* Save config */ false);
 
 		UMediaViewerLibraryIni& Ini = UMediaViewerLibraryIni::Get();
-
-		Ini.SaveLibrary(
-			StaticCastSharedRef<FMediaViewerLibrary>(
-				MediaViewerTab->GetLibrary()
-			)
-		);
-
+		Ini.SaveLibrary(MediaViewerTab->GetLibrary());
 		Ini.SaveConfig();
 	}
 }
@@ -448,7 +489,6 @@ TSharedPtr<SMediaViewerTab> FMediaViewerModule::SpawnTab() const
 TSharedRef<SDockTab> FMediaViewerModule::CreateTab(const FSpawnTabArgs& InArgs)
 {
 	TSharedRef<SMediaViewerTab> ViewerTab = CreateMediaViewer({});
-	UMediaViewerLibraryIni::Get().LoadLibrary(StaticCastSharedRef<FMediaViewerLibrary>(ViewerTab->GetLibrary()));
 
 	return SNew(SDockTab)
 		.Label(LOCTEXT("MediaViewerTitle", "Media Viewer"))
diff --git a/Engine/Plugins/Experimental/MediaViewer/Source/MediaViewer/Private/MediaViewerModule.h b/Engine/Plugins/Experimental/MediaViewer/Source/MediaViewer/Private/MediaViewerModule.h
--- a/Engine/Plugins/Experimental/MediaViewer/Source/MediaViewer/Private/MediaViewerModule.h
+++ b/Engine/Plugins/Experimental/MediaViewer/Source/MediaViewer/Private/MediaViewerModule.h
@@ -11,6 +11,7 @@ class SDockTab;
 
 namespace UE::MediaViewer
 {
+	class IMediaViewerLibrary;
 	class SMediaViewerTab;
 	struct FMediaViewerArgs;
 }
@@ -24,6 +25,9 @@ class FMediaViewerModule : public IMediaViewerModule
 {
 public:
 	static const FLazyName TabId;
+	static const FLazyName StandaloneTabId;
+
+	FMediaViewerModule();
 
 	//~ Begin IMediaViewerModule
 	virtual bool IsFactoryRegistered(FName InFactoryName) const override;
@@ -33,7 +37,7 @@ public:
 	virtual void UnregisterFactory(FName InFactoryName) override;
 	virtual bool OpenTab() override;
 	virtual bool OpenTab(const FMediaViewerArgs& InMediaViewerArgs) override;
-	virtual TSharedPtr<IMediaViewerLibrary> GetLibrary() const;
+	virtual TSharedRef<IMediaViewerLibrary> GetLibrary() const;
 	virtual TSharedPtr<FMediaViewerLibraryItem> CreateLibraryItem(const FAssetData& InAssetData) const override;
 	virtual TSharedPtr<FMediaViewerLibraryItem> CreateLibraryItem(TNotNull<UObject*> InObject) const override;
 	virtual TSharedPtr<FMediaViewerLibraryItem> CreateLibraryItem(FName InItemType, const FMediaViewerLibraryItem& InSavedItem) const override;
@@ -45,6 +49,7 @@ public:
 	virtual bool ClearImage(EMediaImageViewerPosition InPosition) override;
 	virtual TSharedRef<SMediaViewerTab> CreateMediaViewer(const FMediaViewerArgs& InMediaViewerArgs) override;
 	virtual TSharedRef<IMediaViewerLibraryWidget> CreateLibrary(const IMediaViewerLibraryWidget::FArgs& InArgs) override;
+	virtual TSharedPtr<SDockTab> OpenStandaloneTab();
 	//~ End IMediaViewerModule
 
 	//~ Begin IModuleInterface
@@ -56,6 +61,7 @@ private:
 	static void OnTabClosed(TSharedRef<SDockTab> InDockTab);
 
 	TMap<FName, TSharedRef<IMediaImageViewerFactory>> Factories;
+	TSharedRef<IMediaViewerLibrary> Library;
 
 	void RegisterLevelEditorMenuItems();
 	void UnregisterLevelEditorMenuItems();
diff --git a/Engine/Plugins/Experimental/MediaViewer/Source/MediaViewer/Private/Widgets/MediaViewerLibraryItemAssetDragDropOperation.cpp b/Engine/Plugins/Experimental/MediaViewer/Source/MediaViewer/Private/Widgets/MediaViewerLibraryItemAssetDragDropOperation.cpp
--- a/Engine/Plugins/Experimental/MediaViewer/Source/MediaViewer/Private/Widgets/MediaViewerLibraryItemAssetDragDropOperation.cpp
+++ b/Engine/Plugins/Experimental/MediaViewer/Source/MediaViewer/Private/Widgets/MediaViewerLibraryItemAssetDragDropOperation.cpp
@@ -20,7 +20,7 @@
 namespace UE::MediaViewer::Private
 {
 
-FMediaViewerLibraryItemAssetDragDropOperation::FMediaViewerLibraryItemAssetDragDropOperation(const TSharedRef<FMediaViewerLibrary>& InLibrary,
+FMediaViewerLibraryItemAssetDragDropOperation::FMediaViewerLibraryItemAssetDragDropOperation(const TSharedRef<IMediaViewerLibrary>& InLibrary,
 	const IMediaViewerLibrary::FGroupItem& InGroupItem)
 	: GroupItem(InGroupItem)
 	, Decorator(SNew(SBorder))
diff --git a/Engine/Plugins/Experimental/MediaViewer/Source/MediaViewer/Private/Widgets/MediaViewerLibraryItemAssetDragDropOperation.h b/Engine/Plugins/Experimental/MediaViewer/Source/MediaViewer/Private/Widgets/MediaViewerLibraryItemAssetDragDropOperation.h
--- a/Engine/Plugins/Experimental/MediaViewer/Source/MediaViewer/Private/Widgets/MediaViewerLibraryItemAssetDragDropOperation.h
+++ b/Engine/Plugins/Experimental/MediaViewer/Source/MediaViewer/Private/Widgets/MediaViewerLibraryItemAssetDragDropOperation.h
@@ -22,14 +22,12 @@ namespace UE::MediaViewer
 namespace UE::MediaViewer::Private
 {
 
-class FMediaViewerLibrary;
-
 class FMediaViewerLibraryItemAssetDragDropOperation : public FAssetDragDropOp
 {
 public:
 	DRAG_DROP_OPERATOR_TYPE(FMediaViewerLibraryItemAssetDragDropOperation, FAssetDragDropOp)
 
-	FMediaViewerLibraryItemAssetDragDropOperation(const TSharedRef<FMediaViewerLibrary>& InLibrary,
+	FMediaViewerLibraryItemAssetDragDropOperation(const TSharedRef<IMediaViewerLibrary>& InLibrary,
 		const IMediaViewerLibrary::FGroupItem& InGroupItem);
 	
 	virtual ~FMediaViewerLibraryItemAssetDragDropOperation() override = default;
diff --git a/Engine/Plugins/Experimental/MediaViewer/Source/MediaViewer/Private/Widgets/MediaViewerLibraryItemDragDropOperation.cpp b/Engine/Plugins/Experimental/MediaViewer/Source/MediaViewer/Private/Widgets/MediaViewerLibraryItemDragDropOperation.cpp
--- a/Engine/Plugins/Experimental/MediaViewer/Source/MediaViewer/Private/Widgets/MediaViewerLibraryItemDragDropOperation.cpp
+++ b/Engine/Plugins/Experimental/MediaViewer/Source/MediaViewer/Private/Widgets/MediaViewerLibraryItemDragDropOperation.cpp
@@ -20,7 +20,7 @@
 namespace UE::MediaViewer::Private
 {
 
-FMediaViewerLibraryItemDragDropOperation::FMediaViewerLibraryItemDragDropOperation(const TSharedRef<FMediaViewerLibrary>& InLibrary,
+FMediaViewerLibraryItemDragDropOperation::FMediaViewerLibraryItemDragDropOperation(const TSharedRef<IMediaViewerLibrary>& InLibrary,
 	const IMediaViewerLibrary::FGroupItem& InGroupItem)
 	: GroupItem(InGroupItem)
 	, Decorator(SNew(SBorder))
diff --git a/Engine/Plugins/Experimental/MediaViewer/Source/MediaViewer/Private/Widgets/MediaViewerLibraryItemDragDropOperation.h b/Engine/Plugins/Experimental/MediaViewer/Source/MediaViewer/Private/Widgets/MediaViewerLibraryItemDragDropOperation.h
--- a/Engine/Plugins/Experimental/MediaViewer/Source/MediaViewer/Private/Widgets/MediaViewerLibraryItemDragDropOperation.h
+++ b/Engine/Plugins/Experimental/MediaViewer/Source/MediaViewer/Private/Widgets/MediaViewerLibraryItemDragDropOperation.h
@@ -22,14 +22,12 @@ namespace UE::MediaViewer
 namespace UE::MediaViewer::Private
 {
 
-class FMediaViewerLibrary;
-
 class FMediaViewerLibraryItemDragDropOperation : public FDecoratedDragDropOp
 {
 public:
 	DRAG_DROP_OPERATOR_TYPE(FMediaViewerLibraryItemDragDropOperation, FDecoratedDragDropOp)
 
-	FMediaViewerLibraryItemDragDropOperation(const TSharedRef<FMediaViewerLibrary>& InLibrary,
+	FMediaViewerLibraryItemDragDropOperation(const TSharedRef<IMediaViewerLibrary>& InLibrary,
 		const IMediaViewerLibrary::FGroupItem& InGroupItem);
 	
 	virtual ~FMediaViewerLibraryItemDragDropOperation() override = default;
diff --git a/Engine/Plugins/Experimental/MediaViewer/Source/MediaViewer/Private/Widgets/SMediaViewer.cpp b/Engine/Plugins/Experimental/MediaViewer/Source/MediaViewer/Private/Widgets/SMediaViewer.cpp
--- a/Engine/Plugins/Experimental/MediaViewer/Source/MediaViewer/Private/Widgets/SMediaViewer.cpp
+++ b/Engine/Plugins/Experimental/MediaViewer/Source/MediaViewer/Private/Widgets/SMediaViewer.cpp
@@ -157,7 +157,7 @@ void SMediaViewer::Construct(const FArguments& InArgs, const TSharedRef<SMediaVi
 	CreateView();
 }
 
-TSharedRef<FMediaViewerLibrary> SMediaViewer::GetLibrary() const
+TSharedRef<IMediaViewerLibrary> SMediaViewer::GetLibrary() const
 {
 	return Library->GetLibrary();
 }
@@ -240,6 +240,8 @@ void SMediaViewer::BindCommands()
 		FExecuteAction::CreateSP(this, &SMediaViewer::TakeSnapshot),
 		FCanExecuteAction::CreateSP(this, &SMediaViewer::CanTakeSnapshot)
 	);
+
+	CommandList->MapAction(Commands.OpenNewTab, FExecuteAction::CreateSP(this, &SMediaViewer::OpenNewTab));
 }
 
 bool SMediaViewer::IsShowingBothImageViewers() const
@@ -755,7 +757,7 @@ TSharedRef<FMediaViewerDelegates> SMediaViewer::GetDelegates() const
 	return Delegates.ToSharedRef();
 }
 
-void SMediaViewer::SaveStates() const
+void SMediaViewer::SaveStates(bool bInSaveConfig) const
 {
 	FMediaViewerState State;
 	State.ViewerSettings = MediaViewerSettings;
@@ -787,7 +789,11 @@ void SMediaViewer::SaveStates() const
 
 	UMediaViewerLibraryIni& Ini = UMediaViewerLibraryIni::Get();
 	Ini.SetSavedStates({State});
-	Ini.SaveConfig();
+
+	if (bInSaveConfig)
+	{
+		Ini.SaveConfig();
+	}
 }
 
 void SMediaViewer::LoadState(int32 InIndex)
@@ -1187,7 +1193,7 @@ void SMediaViewer::TakeSnapshot()
 		TNotNull<UTexture2D*>(Texture)
 	);
 
-	TSharedRef<FMediaViewerLibrary> LibraryActual = GetLibrary();
+	TSharedRef<IMediaViewerLibrary> LibraryActual = GetLibrary();
 	LibraryActual->AddItemToGroup(SnapshotItem, LibraryActual->GetSnapshotsGroupId());
 
 	// Make a toast!
@@ -1260,6 +1266,11 @@ void SMediaViewer::DrawImageViewers(UCanvas* InCanvas) const
 	}
 }
 
+void SMediaViewer::OpenNewTab()
+{
+	IMediaViewerModule::Get().OpenStandaloneTab();
+}
+
 } // UE::MediaViewer::Private
 
 #undef LOCTEXT_NAMESPACE
diff --git a/Engine/Plugins/Experimental/MediaViewer/Source/MediaViewer/Private/Widgets/SMediaViewer.h b/Engine/Plugins/Experimental/MediaViewer/Source/MediaViewer/Private/Widgets/SMediaViewer.h
--- a/Engine/Plugins/Experimental/MediaViewer/Source/MediaViewer/Private/Widgets/SMediaViewer.h
+++ b/Engine/Plugins/Experimental/MediaViewer/Source/MediaViewer/Private/Widgets/SMediaViewer.h
@@ -22,13 +22,13 @@ namespace UE::MediaViewer
 {
 	class FMediaImageStatusBarExtender;
 	class FMediaImageViewer;
+	class IMediaViewerLibrary;
 	class SMediaViewerTab;
 	struct FMediaViewerDelegates;
 
 	namespace Private
 	{
 		class FMediaImageViewerViewportClient;
-		class FMediaViewerLibrary;
 		class SMediaImageViewerOverlay;
 		class SMediaImageViewerToolbar;
 		class SMediaViewerLibraryPrivate;
@@ -68,7 +68,7 @@ public:
 
 	const FMediaViewerArgs& GetArgs() const { return MediaViewerArgs; }
 
-	TSharedRef<FMediaViewerLibrary> GetLibrary() const;
+	TSharedRef<IMediaViewerLibrary> GetLibrary() const;
 
 	TSharedPtr<FMediaImageViewer> GetImageViewer(EMediaImageViewerPosition InPosition) const;
 
@@ -99,7 +99,7 @@ public:
 	/** Returns the bound delegates of this viewer. */
 	TSharedRef<FMediaViewerDelegates> GetDelegates() const;
 
-	void SaveStates() const;
+	void SaveStates(bool bInSaveConfig) const;
 
 	void LoadState(int32 InIndex);
 
@@ -246,6 +246,8 @@ protected:
 	void TakeSnapshot();
 
 	void DrawImageViewers(UCanvas* InCanvas) const;
+	
+	void OpenNewTab();
 };
 
 } // UE::MediaViewer::Private
diff --git a/Engine/Plugins/Experimental/MediaViewer/Source/MediaViewer/Private/Widgets/SMediaViewerLibrary.cpp b/Engine/Plugins/Experimental/MediaViewer/Source/MediaViewer/Private/Widgets/SMediaViewerLibrary.cpp
--- a/Engine/Plugins/Experimental/MediaViewer/Source/MediaViewer/Private/Widgets/SMediaViewerLibrary.cpp
+++ b/Engine/Plugins/Experimental/MediaViewer/Source/MediaViewer/Private/Widgets/SMediaViewerLibrary.cpp
@@ -3,6 +3,7 @@
 #include "SMediaViewerLibrary.h"
 
 #include "DetailLayoutBuilder.h"
+#include "IMediaViewerModule.h"
 #include "Library/MediaViewerLibrary.h"
 #include "Library/MediaViewerLibraryDynamicGroup.h"
 #include "Library/MediaViewerLibraryGroup.h"
@@ -21,6 +22,14 @@ SMediaViewerLibrary::SMediaViewerLibrary()
 {
 }
 
+SMediaViewerLibrary::~SMediaViewerLibrary()
+{
+	if (Library.IsValid())
+	{
+		Library->GetOnChanged().RemoveAll(this);
+	}
+}
+
 void SMediaViewerLibrary::PrivateRegisterAttributes(FSlateAttributeDescriptor::FInitializer&)
 {
 }
@@ -30,8 +39,8 @@ void SMediaViewerLibrary::Construct(const FArguments& InArgs, const IMediaViewer
 {
 	GroupFilter = InMediaViewerLibraryArgs.GroupFilter;
 	Delegates = InDelegates;
-	Library = MakeShared<FMediaViewerLibrary>(InMediaViewerLibraryArgs.LibraryArgs);
-	Library->GetOnChanged().BindSP(this, &SMediaViewerLibrary::OnLibraryChanged);
+	Library = IMediaViewerModule::Get().GetLibrary();
+	Library->GetOnChanged().AddSP(this, &SMediaViewerLibrary::OnLibraryChanged);
 
 	SetVisibility(EVisibility::Visible);
 
@@ -60,11 +69,6 @@ TSharedRef<IMediaViewerLibrary> SMediaViewerLibrary::GetLibrary() const
 	return Library.ToSharedRef();
 }
 
-TSharedRef<FMediaViewerLibrary> SMediaViewerLibrary::GetLibraryImpl() const
-{
-	return Library.ToSharedRef();
-}
-
 FReply SMediaViewerLibrary::OnMouseButtonDown(const FGeometry& InMyGeometry, const FPointerEvent& InMouseEvent)
 {
 	// Make sure sidebar blocks widget below it.
@@ -90,13 +94,11 @@ void SMediaViewerLibrary::OnLibraryChanged(TSharedRef<IMediaViewerLibrary> InLib
 
 void SMediaViewerLibrary::UpdateGroups()
 {
-	TSharedRef<FMediaViewerLibrary> LibraryImpl = GetLibraryImpl();
-
-	Groups.Empty(LibraryImpl->GetGroups().Num());
+	Groups.Empty(Library->GetGroups().Num());
 
 	TSharedRef<const IMediaViewerLibrary> ConstLibrary = Library.ToSharedRef();
 
-	for (const TSharedRef<const FMediaViewerLibraryGroup>& Group : LibraryImpl->GetGroups())
+	for (const TSharedRef<const FMediaViewerLibraryGroup>& Group : Library->GetGroups())
 	{
 		if (!GroupFilter.IsBound() || GroupFilter.Execute(ConstLibrary, Group))
 		{
@@ -123,9 +125,7 @@ void SMediaViewerLibrary::OnGetChildren(IMediaViewerLibrary::FGroupItem InParent
 		return;
 	}
 
-	TSharedRef<FMediaViewerLibrary> LibraryImpl = GetLibraryImpl();
-
-	TSharedPtr<FMediaViewerLibraryGroup> Group = LibraryImpl->GetGroup(InParent.GroupId);
+	TSharedPtr<FMediaViewerLibraryGroup> Group = Library->GetGroup(InParent.GroupId);
 
 	if (!Group.IsValid())
 	{
@@ -155,9 +155,7 @@ void SMediaViewerLibrary::OnGroupExpanded(IMediaViewerLibrary::FGroupItem InGrou
 		return;
 	}
 
-	TSharedRef<FMediaViewerLibrary> LibraryImpl = GetLibraryImpl();
-
-	TSharedPtr<FMediaViewerLibraryGroup> Group = LibraryImpl->GetGroup(InGroup.GroupId);
+	TSharedPtr<FMediaViewerLibraryGroup> Group = Library->GetGroup(InGroup.GroupId);
 
 	if (Group.IsValid() && Group->IsDynamic())
 	{
diff --git a/Engine/Plugins/Experimental/MediaViewer/Source/MediaViewer/Private/Widgets/SMediaViewerLibrary.h b/Engine/Plugins/Experimental/MediaViewer/Source/MediaViewer/Private/Widgets/SMediaViewerLibrary.h
--- a/Engine/Plugins/Experimental/MediaViewer/Source/MediaViewer/Private/Widgets/SMediaViewerLibrary.h
+++ b/Engine/Plugins/Experimental/MediaViewer/Source/MediaViewer/Private/Widgets/SMediaViewerLibrary.h
@@ -21,8 +21,6 @@ namespace UE::MediaViewer
 namespace UE::MediaViewer::Private
 {
 
-class FMediaViewerLibrary;
-
 /**
  * Implementation of @see IMediaViewerLibraryWidget.
  */
@@ -37,12 +35,10 @@ public:
 
 	SMediaViewerLibrary();
 
-	virtual ~SMediaViewerLibrary() override = default;
+	virtual ~SMediaViewerLibrary() override;
 
 	void Construct(const FArguments& InArgs, const FArgs& InMediaViewerLibraryArgs, const TSharedRef<FMediaViewerDelegates>& InDelegates);
 
-	virtual TSharedRef<FMediaViewerLibrary> GetLibraryImpl() const;
-
 	void OnLibraryChanged(TSharedRef<IMediaViewerLibrary> InLibrary, IMediaViewerLibrary::EChangeType InChangeType);
 
 	//~ Begin IMediaViewerLibraryWidget
@@ -56,7 +52,7 @@ public:
 
 protected:
 	TSharedPtr<FMediaViewerDelegates> Delegates;
-	TSharedPtr<FMediaViewerLibrary> Library;
+	TSharedPtr<IMediaViewerLibrary> Library;
 	FGroupFilter GroupFilter;
 	TArray<IMediaViewerLibrary::FGroupItem> Groups;
 	TSharedPtr<STreeView<IMediaViewerLibrary::FGroupItem>> TreeView;
diff --git a/Engine/Plugins/Experimental/MediaViewer/Source/MediaViewer/Private/Widgets/SMediaViewerLibraryGroup.cpp b/Engine/Plugins/Experimental/MediaViewer/Source/MediaViewer/Private/Widgets/SMediaViewerLibraryGroup.cpp
--- a/Engine/Plugins/Experimental/MediaViewer/Source/MediaViewer/Private/Widgets/SMediaViewerLibraryGroup.cpp
+++ b/Engine/Plugins/Experimental/MediaViewer/Source/MediaViewer/Private/Widgets/SMediaViewerLibraryGroup.cpp
@@ -37,7 +37,7 @@ void SMediaViewerLibraryGroup::PrivateRegisterAttributes(FSlateAttributeDescript
 }
 
 void SMediaViewerLibraryGroup::Construct(const FArguments& InArgs, const TSharedRef<STableViewBase>& InOwningTable,
-	const TSharedRef<FMediaViewerLibrary>& InLibrary, const FGuid& InGroupId)
+	const TSharedRef<IMediaViewerLibrary>& InLibrary, const FGuid& InGroupId)
 {
 	LibraryWeak = InLibrary;
 	GroupId = InGroupId;
@@ -155,14 +155,16 @@ void SMediaViewerLibraryGroup::OnDragEnter(const FGeometry& InMyGeometry, const
 {
 	FSuperType::OnDragEnter(InMyGeometry, InEvent);
 
-	TSharedPtr<FMediaViewerLibrary> Library = LibraryWeak.Pin();
+	TSharedPtr<IMediaViewerLibrary> Library = LibraryWeak.Pin();
 
 	if (!Library.IsValid())
 	{
 		return;
 	}
 
-	if (!Library->CanDragDropGroup(GroupId))
+	TSharedPtr<FMediaViewerLibrary> LibraryImpl = StaticCastSharedPtr<FMediaViewerLibrary>(Library);
+
+	if (!LibraryImpl->CanDragDropGroup(GroupId))
 	{
 		return;
 	}
@@ -193,14 +195,16 @@ FReply SMediaViewerLibraryGroup::OnDrop(const FGeometry& InMyGeometry, const FDr
 {
 	ItemDropZone = TOptional<EItemDropZone>();
 
-	TSharedPtr<FMediaViewerLibrary> Library = LibraryWeak.Pin();
+	TSharedPtr<IMediaViewerLibrary> Library = LibraryWeak.Pin();
 
 	if (!Library.IsValid())
 	{
 		return FSuperType::OnDrop(InMyGeometry, InEvent);
 	}
 
-	if (!Library->CanDragDropGroup(GroupId))
+	TSharedPtr<FMediaViewerLibrary> LibraryImpl = StaticCastSharedPtr<FMediaViewerLibrary>(Library);
+
+	if (!LibraryImpl->CanDragDropGroup(GroupId))
 	{
 		return FSuperType::OnDrop(InMyGeometry, InEvent);
 	}
@@ -223,7 +227,7 @@ FReply SMediaViewerLibraryGroup::OnDrop(const FGeometry& InMyGeometry, const FDr
 
 FText SMediaViewerLibraryGroup::GetGroupName() const
 {
-	if (TSharedPtr<FMediaViewerLibrary> Library = LibraryWeak.Pin())
+	if (TSharedPtr<IMediaViewerLibrary> Library = LibraryWeak.Pin())
 	{
 		if (TSharedPtr<FMediaViewerLibraryGroup> Group = Library->GetGroup(GroupId))
 		{
@@ -236,7 +240,7 @@ FText SMediaViewerLibraryGroup::GetGroupName() const
 
 FReply SMediaViewerLibraryGroup::OnRemoveButtonClicked()
 {
-	TSharedPtr<FMediaViewerLibrary> Library = LibraryWeak.Pin();
+	TSharedPtr<IMediaViewerLibrary> Library = LibraryWeak.Pin();
 
 	if (!Library.IsValid())
 	{
@@ -264,7 +268,7 @@ FReply SMediaViewerLibraryGroup::OnRemoveButtonClicked()
 
 FReply SMediaViewerLibraryGroup::OnRefreshButtonClicked()
 {
-	TSharedPtr<FMediaViewerLibrary> Library = LibraryWeak.Pin();
+	TSharedPtr<IMediaViewerLibrary> Library = LibraryWeak.Pin();
 
 	if (!Library.IsValid())
 	{
@@ -285,7 +289,7 @@ FReply SMediaViewerLibraryGroup::OnRefreshButtonClicked()
 
 bool SMediaViewerLibraryGroup::CanAcceptLibraryItem(const IMediaViewerLibrary::FGroupItem& InDraggedGroupItem) const
 {
-	TSharedPtr<FMediaViewerLibrary> Library = LibraryWeak.Pin();
+	TSharedPtr<IMediaViewerLibrary> Library = LibraryWeak.Pin();
 
 	if (!Library.IsValid())
 	{
@@ -297,12 +301,14 @@ bool SMediaViewerLibraryGroup::CanAcceptLibraryItem(const IMediaViewerLibrary::F
 		return false;
 	}
 
-	if (!Library->CanDragDropGroup(GroupId))
+	TSharedPtr<FMediaViewerLibrary> LibraryImpl = StaticCastSharedPtr<FMediaViewerLibrary>(Library);
+
+	if (!LibraryImpl->CanDragDropGroup(GroupId))
 	{
 		return false;
 	}
 
-	if (!Library->CanDragDropItem(InDraggedGroupItem))
+	if (!LibraryImpl->CanDragDropItem(InDraggedGroupItem))
 	{
 		return false;
 	}
@@ -318,7 +324,7 @@ void SMediaViewerLibraryGroup::OnLibraryItemDropped(const IMediaViewerLibrary::F
 	}
 
 	// Validated by CanAcceptLibraryItem
-	TSharedPtr<FMediaViewerLibrary> Library = LibraryWeak.Pin();
+	TSharedPtr<IMediaViewerLibrary> Library = LibraryWeak.Pin();
 
 	if (GroupId == InDroppedGroupItem.GroupId)
 	{
@@ -344,7 +350,7 @@ void SMediaViewerLibraryGroup::OnLibraryItemDropped(const IMediaViewerLibrary::F
 
 bool SMediaViewerLibraryGroup::CanAcceptAssets(TConstArrayView<FAssetData> InAssetData) const
 {
-	TSharedPtr<FMediaViewerLibrary> Library = LibraryWeak.Pin();
+	TSharedPtr<IMediaViewerLibrary> Library = LibraryWeak.Pin();
 
 	if (!Library.IsValid())
 	{
@@ -376,7 +382,7 @@ bool SMediaViewerLibraryGroup::CanAcceptAssets(TConstArrayView<FAssetData> InAss
 void SMediaViewerLibraryGroup::OnAssetsDropped(TConstArrayView<FAssetData> InAssetData)
 {
 	// Already verified at this point
-	TSharedPtr<FMediaViewerLibrary> Library = LibraryWeak.Pin();
+	TSharedPtr<IMediaViewerLibrary> Library = LibraryWeak.Pin();
 	TSharedPtr<FMediaViewerLibraryGroup> Group = Library->GetGroup(GroupId);
 	
 	IMediaViewerModule& Module = IMediaViewerModule::Get();
diff --git a/Engine/Plugins/Experimental/MediaViewer/Source/MediaViewer/Private/Widgets/SMediaViewerLibraryGroup.h b/Engine/Plugins/Experimental/MediaViewer/Source/MediaViewer/Private/Widgets/SMediaViewerLibraryGroup.h
--- a/Engine/Plugins/Experimental/MediaViewer/Source/MediaViewer/Private/Widgets/SMediaViewerLibraryGroup.h
+++ b/Engine/Plugins/Experimental/MediaViewer/Source/MediaViewer/Private/Widgets/SMediaViewerLibraryGroup.h
@@ -4,6 +4,7 @@
 
 #include "Widgets/Views/STableRow.h"
 
+#include "Library/IMediaViewerLibrary.h"
 #include "Misc/Guid.h"
 #include "Widgets/IMediaViewerLibraryWidget.h"
 #include "Widgets/MediaViewerLibraryGroupItem.h"
@@ -20,8 +21,6 @@ namespace UE::MediaViewer
 namespace UE::MediaViewer::Private
 {
 
-class FMediaViewerLibrary;
-
 class SMediaViewerLibraryGroup : public STableRow<IMediaViewerLibrary::FGroupItem>
 {
 	using FSuperType = STableRow<IMediaViewerLibrary::FGroupItem>;
@@ -38,7 +37,7 @@ public:
 	virtual ~SMediaViewerLibraryGroup() override = default;
 
 	void Construct(const FArguments& InArgs, const TSharedRef<STableViewBase>& InOwningTable, 
-		const TSharedRef<FMediaViewerLibrary>& InLibrary, const FGuid& InGroupId);
+		const TSharedRef<IMediaViewerLibrary>& InLibrary, const FGuid& InGroupId);
 
 	//~ Begin STableRow
 	virtual void OnDragEnter(const FGeometry& InMyGeometry, const FDragDropEvent& InEvent) override;
@@ -46,7 +45,7 @@ public:
 	//~ End STableRow
 
 protected:
-	TWeakPtr<FMediaViewerLibrary> LibraryWeak;
+	TWeakPtr<IMediaViewerLibrary> LibraryWeak;
 	FGuid GroupId;
 
 	FText GetGroupName() const;
diff --git a/Engine/Plugins/Experimental/MediaViewer/Source/MediaViewer/Private/Widgets/SMediaViewerLibraryItem.cpp b/Engine/Plugins/Experimental/MediaViewer/Source/MediaViewer/Private/Widgets/SMediaViewerLibraryItem.cpp
--- a/Engine/Plugins/Experimental/MediaViewer/Source/MediaViewer/Private/Widgets/SMediaViewerLibraryItem.cpp
+++ b/Engine/Plugins/Experimental/MediaViewer/Source/MediaViewer/Private/Widgets/SMediaViewerLibraryItem.cpp
@@ -41,7 +41,7 @@ void SMediaViewerLibraryItem::PrivateRegisterAttributes(FSlateAttributeDescripto
 }
 
 void SMediaViewerLibraryItem::Construct(const FArguments& InArgs, const TSharedRef<STableViewBase>& InOwningTable, 
-	const TSharedRef<FMediaViewerLibrary>& InLibrary, const IMediaViewerLibrary::FGroupItem& InGroupItem,
+	const TSharedRef<IMediaViewerLibrary>& InLibrary, const IMediaViewerLibrary::FGroupItem& InGroupItem,
 	const TSharedRef<FMediaViewerDelegates>& InDelegates)
 {
 	LibraryWeak = InLibrary;
@@ -306,7 +306,7 @@ void SMediaViewerLibraryItem::OnUseButtonClicked(ECheckBoxState InCheckState, EM
 
 FReply SMediaViewerLibraryItem::OnRemoveButtonClicked()
 {
-	TSharedPtr<FMediaViewerLibrary> Library = LibraryWeak.Pin();
+	TSharedPtr<IMediaViewerLibrary> Library = LibraryWeak.Pin();
 
 	if (!Library.IsValid())
 	{
@@ -325,7 +325,7 @@ FReply SMediaViewerLibraryItem::OnDragDetected(const FGeometry& InMyGeometry, co
 		return FReply::Handled();
 	}
 
-	TSharedPtr<FMediaViewerLibrary> Library = LibraryWeak.Pin();
+	TSharedPtr<IMediaViewerLibrary> Library = LibraryWeak.Pin();
 
 	if (!Library.IsValid())
 	{
@@ -353,7 +353,7 @@ bool SMediaViewerLibraryItem::CanAcceptLibraryItem(const IMediaViewerLibrary::FG
 		return false;
 	}
 
-	TSharedPtr<FMediaViewerLibrary> Library = LibraryWeak.Pin();
+	TSharedPtr<IMediaViewerLibrary> Library = LibraryWeak.Pin();
 
 	if (!Library.IsValid())
 	{
@@ -365,12 +365,14 @@ bool SMediaViewerLibraryItem::CanAcceptLibraryItem(const IMediaViewerLibrary::FG
 		return false;
 	}
 
-	if (!Library->CanDragDropGroup(GroupItem.GroupId))
+	TSharedPtr<FMediaViewerLibrary> LibraryImpl = StaticCastSharedPtr<FMediaViewerLibrary>(Library);
+
+	if (!LibraryImpl->CanDragDropGroup(GroupItem.GroupId))
 	{
 		return false;
 	}
 
-	if (!Library->CanDragDropGroup(InDraggedGroupItem.GroupId))
+	if (!LibraryImpl->CanDragDropGroup(InDraggedGroupItem.GroupId))
 	{
 		return false;
 	}
@@ -400,7 +402,7 @@ void SMediaViewerLibraryItem::OnLibraryItemDropped(const IMediaViewerLibrary::FG
 	}
 
 	// Validated by CanAcceptLibraryItem
-	TSharedPtr<FMediaViewerLibrary> Library = LibraryWeak.Pin();
+	TSharedPtr<IMediaViewerLibrary> Library = LibraryWeak.Pin();
 	TSharedPtr<FMediaViewerLibraryGroup> MyGroup = Library->GetGroup(GroupItem.GroupId);
 	const int32 MyGroupIndex = MyGroup->FindItemIndex(GroupItem.ItemId);
 
@@ -505,7 +507,7 @@ FReply SMediaViewerLibraryItem::OnDrop(const FGeometry& InMyGeometry, const FDra
 
 bool SMediaViewerLibraryItem::CanAcceptAssets(TConstArrayView<FAssetData> InAssetData) const
 {
-	TSharedPtr<FMediaViewerLibrary> Library = LibraryWeak.Pin();
+	TSharedPtr<IMediaViewerLibrary> Library = LibraryWeak.Pin();
 
 	if (!Library.IsValid())
 	{
@@ -537,7 +539,7 @@ bool SMediaViewerLibraryItem::CanAcceptAssets(TConstArrayView<FAssetData> InAsse
 void SMediaViewerLibraryItem::OnAssetsDropped(TConstArrayView<FAssetData> InAssetData)
 {
 	// Already verified at this point
-	TSharedPtr<FMediaViewerLibrary> Library = LibraryWeak.Pin();
+	TSharedPtr<IMediaViewerLibrary> Library = LibraryWeak.Pin();
 
 	TSharedPtr<FMediaViewerLibraryGroup> Group = Library->GetGroup(GroupItem.GroupId);
 
diff --git a/Engine/Plugins/Experimental/MediaViewer/Source/MediaViewer/Private/Widgets/SMediaViewerLibraryItem.h b/Engine/Plugins/Experimental/MediaViewer/Source/MediaViewer/Private/Widgets/SMediaViewerLibraryItem.h
--- a/Engine/Plugins/Experimental/MediaViewer/Source/MediaViewer/Private/Widgets/SMediaViewerLibraryItem.h
+++ b/Engine/Plugins/Experimental/MediaViewer/Source/MediaViewer/Private/Widgets/SMediaViewerLibraryItem.h
@@ -4,6 +4,7 @@
 
 #include "Widgets/Views/STableRow.h"
 
+#include "Library/IMediaViewerLibrary.h"
 #include "Misc/Guid.h"
 #include "Widgets/IMediaViewerLibraryWidget.h"
 #include "Widgets/MediaViewerLibraryGroupItem.h"
@@ -24,8 +25,6 @@ namespace UE::MediaViewer
 namespace UE::MediaViewer::Private
 {
 
-class FMediaViewerLibrary;
-	
 class SMediaViewerLibraryItem : public STableRow<IMediaViewerLibrary::FGroupItem>
 {
 	using FSuperType = STableRow<IMediaViewerLibrary::FGroupItem>;
@@ -42,7 +41,7 @@ public:
 	virtual ~SMediaViewerLibraryItem() override = default;
 
 	void Construct(const FArguments& InArgs, const TSharedRef<STableViewBase>& InOwningTable,
-		const TSharedRef<FMediaViewerLibrary>& InLibrary, const IMediaViewerLibrary::FGroupItem& InGroupItem, 
+		const TSharedRef<IMediaViewerLibrary>& InLibrary, const IMediaViewerLibrary::FGroupItem& InGroupItem, 
 		const TSharedRef<FMediaViewerDelegates>& InDelegates);
 
 	//~ Begin STableRow
@@ -51,7 +50,7 @@ public:
 	//~ End STableRow
 
 protected:
-	TWeakPtr<FMediaViewerLibrary> LibraryWeak;
+	TWeakPtr<IMediaViewerLibrary> LibraryWeak;
 	IMediaViewerLibrary::FGroupItem GroupItem;
 	TSharedPtr<FMediaViewerDelegates> Delegates;
 	TSharedPtr<FMediaViewerLibraryItem> Item;
diff --git a/Engine/Plugins/Experimental/MediaViewer/Source/MediaViewer/Private/Widgets/SMediaViewerLibraryPrivate.cpp b/Engine/Plugins/Experimental/MediaViewer/Source/MediaViewer/Private/Widgets/SMediaViewerLibraryPrivate.cpp
--- a/Engine/Plugins/Experimental/MediaViewer/Source/MediaViewer/Private/Widgets/SMediaViewerLibraryPrivate.cpp
+++ b/Engine/Plugins/Experimental/MediaViewer/Source/MediaViewer/Private/Widgets/SMediaViewerLibraryPrivate.cpp
@@ -4,9 +4,6 @@
 
 #include "ImageViewers/ColorImageViewer.h"
 #include "LevelEditor.h"
-#include "Library/LevelEditorViewportGroup.h"
-#include "Library/MediaTextureGroup.h"
-#include "Library/MediaTrackGroup.h"
 #include "Library/MediaViewerLibrary.h"
 #include "Library/MediaViewerLibraryGroup.h"
 #include "Library/MediaViewerLibraryIni.h"
@@ -39,17 +36,11 @@ void SMediaViewerLibraryPrivate::Construct(const FArguments& InArgs, const TShar
 	[
 		SAssignNew(Library, SMediaViewerLibrary, Args, InDelegates)
 	];
-
-	TSharedRef<FMediaViewerLibrary> LibraryImpl = GetLibrary();
-
-	LibraryImpl->AddGroup(MakeShared<FLevelEditorViewportGroup>(LibraryImpl));
-	LibraryImpl->AddGroup(MakeShared<FMediaTrackGroup>(LibraryImpl));
-	LibraryImpl->AddGroup(MakeShared<FMediaTextureGroup>(LibraryImpl));
 }
 
-TSharedRef<FMediaViewerLibrary> SMediaViewerLibraryPrivate::GetLibrary() const
+TSharedRef<IMediaViewerLibrary> SMediaViewerLibraryPrivate::GetLibrary() const
 {
-	return Library->GetLibraryImpl();
+	return Library->GetLibrary();
 }
 
 void SMediaViewerLibraryPrivate::OnImageViewerOpened(const TSharedRef<FMediaImageViewer>& InImageViewer)
@@ -87,16 +78,14 @@ void SMediaViewerLibraryPrivate::OnImageViewerOpened(const TSharedRef<FMediaImag
 		}
 	}
 
-	HistoryGroup->RemoveItem(Item->GetId());
+	LibraryImpl->RemoveItemFromGroup({HistoryGroup->GetId(), Item->GetId()});
 
-	HistoryGroup->AddItem(Item->GetId(), 0);
+	LibraryImpl->AddItemToGroup(Item.ToSharedRef(), HistoryGroup->GetId(), /* Add to the top of the group */ 0);
 
 	while (HistoryGroup->GetItems().Num() > UE::MediaViewer::Private::MaxHistoryEntries)
 	{
-		HistoryGroup->RemoveItem(HistoryGroup->GetItems().Last());
+		LibraryImpl->RemoveItemFromGroup({HistoryGroup->GetId(), HistoryGroup->GetItems().Last()});
 	}
-
-	Library->OnLibraryChanged(Library->GetLibrary(), IMediaViewerLibrary::EChangeType::ItemGroupChanged);
 }
 
 } // UE::MediaViewer::Private
diff --git a/Engine/Plugins/Experimental/MediaViewer/Source/MediaViewer/Private/Widgets/SMediaViewerLibraryPrivate.h b/Engine/Plugins/Experimental/MediaViewer/Source/MediaViewer/Private/Widgets/SMediaViewerLibraryPrivate.h
--- a/Engine/Plugins/Experimental/MediaViewer/Source/MediaViewer/Private/Widgets/SMediaViewerLibraryPrivate.h
+++ b/Engine/Plugins/Experimental/MediaViewer/Source/MediaViewer/Private/Widgets/SMediaViewerLibraryPrivate.h
@@ -11,6 +11,7 @@ struct FMediaViewerLibraryGroup;
 namespace UE::MediaViewer
 {
 	class FMediaImageViewer;
+	class IMediaViewerLibrary;
 	enum class EMediaImageViewerPosition : uint8;
 	struct FMediaViewerDelegates;
 }
@@ -18,7 +19,6 @@ namespace UE::MediaViewer
 namespace UE::MediaViewer::Private
 {
 
-class FMediaViewerLibrary;
 class SMediaViewerLibrary;
 
 class SMediaViewerLibraryPrivate : public SCompoundWidget
@@ -36,7 +36,7 @@ public:
 
 	void Construct(const FArguments& InArgs, const TSharedRef<FMediaViewerDelegates>& InDelegates);
 
-	TSharedRef<FMediaViewerLibrary> GetLibrary() const;
+	TSharedRef<IMediaViewerLibrary> GetLibrary() const;
 
 	void OnImageViewerOpened(const TSharedRef<FMediaImageViewer>& InImageViewer);
 
diff --git a/Engine/Plugins/Experimental/MediaViewer/Source/MediaViewer/Private/Widgets/SMediaViewerToolbar.cpp b/Engine/Plugins/Experimental/MediaViewer/Source/MediaViewer/Private/Widgets/SMediaViewerToolbar.cpp
--- a/Engine/Plugins/Experimental/MediaViewer/Source/MediaViewer/Private/Widgets/SMediaViewerToolbar.cpp
+++ b/Engine/Plugins/Experimental/MediaViewer/Source/MediaViewer/Private/Widgets/SMediaViewerToolbar.cpp
@@ -265,7 +265,11 @@ TSharedRef<SWidget> SMediaViewerToolbar::MakeCenterToolbar()
 				TAttribute<FText>(),
 				FSlateIcon(FAppStyle::Get().GetStyleSetName(), "LevelViewport.HighResScreenshot")
 			);
+		}
+		ToolbarBuilder.EndBlockGroup();
 
+		ToolbarBuilder.BeginBlockGroup();
+		{
 			ToolbarBuilder.AddComboButton(
 				FUIAction(),
 				FOnGetContent::CreateSP(this, &SMediaViewerToolbar::GetBackgroundTextureSettingsWidget),
@@ -273,6 +277,14 @@ TSharedRef<SWidget> SMediaViewerToolbar::MakeCenterToolbar()
 				FText::GetEmpty(),
 				FSlateIcon(FAppStyle::Get().GetStyleSetName(), "Icons.Toolbar.Settings")
 			);
+
+			ToolbarBuilder.AddToolBarButton(
+				Commands.OpenNewTab,
+				NAME_None,
+				FText::GetEmpty(),
+				TAttribute<FText>(),
+				FSlateIcon(FAppStyle::Get().GetStyleSetName(), "Icons.OpenInExternalEditor")
+			);
 		}
 		ToolbarBuilder.EndBlockGroup();
 	}
diff --git a/Engine/Plugins/Experimental/MediaViewer/Source/MediaViewer/Public/IMediaViewerModule.h b/Engine/Plugins/Experimental/MediaViewer/Source/MediaViewer/Public/IMediaViewerModule.h
--- a/Engine/Plugins/Experimental/MediaViewer/Source/MediaViewer/Public/IMediaViewerModule.h
+++ b/Engine/Plugins/Experimental/MediaViewer/Source/MediaViewer/Public/IMediaViewerModule.h
@@ -11,6 +11,7 @@
 #include "Templates/SharedPointer.h"
 #include "Widgets/IMediaViewerLibraryWidget.h"
 
+class SDockTab;
 struct FAssetData;
 struct FMediaViewerLibraryItem;
 
@@ -55,7 +56,7 @@ public:
 	virtual bool OpenTab(const FMediaViewerArgs& InMediaViewerArgs) = 0;
 
 	/** Gets the Library for the open tab. */
-	virtual TSharedPtr<IMediaViewerLibrary> GetLibrary() const = 0;
+	virtual TSharedRef<IMediaViewerLibrary> GetLibrary() const = 0;
 
 	virtual TSharedPtr<FMediaViewerLibraryItem> CreateLibraryItem(const FAssetData& InAssetData) const = 0;
 	virtual TSharedPtr<FMediaViewerLibraryItem> CreateLibraryItem(TNotNull<UObject*> InObject) const = 0;
@@ -82,6 +83,11 @@ public:
 
 	/** Creates a Library widget. */
 	virtual TSharedRef<IMediaViewerLibraryWidget> CreateLibrary(const IMediaViewerLibraryWidget::FArgs& InArgs) = 0;
+
+	/**
+	 * Methods for creating new standalone tabs
+	 */
+	virtual TSharedPtr<SDockTab> OpenStandaloneTab() = 0;
 };
 
 } // UE::MediaViewer
diff --git a/Engine/Plugins/Experimental/MediaViewer/Source/MediaViewer/Public/Library/IMediaViewerLibrary.h b/Engine/Plugins/Experimental/MediaViewer/Source/MediaViewer/Public/Library/IMediaViewerLibrary.h
--- a/Engine/Plugins/Experimental/MediaViewer/Source/MediaViewer/Public/Library/IMediaViewerLibrary.h
+++ b/Engine/Plugins/Experimental/MediaViewer/Source/MediaViewer/Public/Library/IMediaViewerLibrary.h
@@ -54,14 +54,7 @@ public:
 		ItemGroupChanged
 	};
 
-	DECLARE_DELEGATE_TwoParams(FOnChanged, TSharedRef<IMediaViewerLibrary> /* This */, EChangeType)
-
-	struct FArgs
-	{
-		FText DefaultGroupName;
-		FText DefaultGroupToolTip;
-		FOnChanged OnChanged;
-	};
+	DECLARE_MULTICAST_DELEGATE_TwoParams(FOnChanged, TSharedRef<IMediaViewerLibrary> /* This */, EChangeType)
 
 	/**
 	 * @returns The id of the default group.
diff --git a/Engine/Plugins/Experimental/MediaViewer/Source/MediaViewer/Public/Widgets/IMediaViewerLibraryWidget.h b/Engine/Plugins/Experimental/MediaViewer/Source/MediaViewer/Public/Widgets/IMediaViewerLibraryWidget.h
--- a/Engine/Plugins/Experimental/MediaViewer/Source/MediaViewer/Public/Widgets/IMediaViewerLibraryWidget.h
+++ b/Engine/Plugins/Experimental/MediaViewer/Source/MediaViewer/Public/Widgets/IMediaViewerLibraryWidget.h
@@ -45,7 +45,6 @@ public:
 
 	struct FArgs
 	{
-		IMediaViewerLibrary::FArgs LibraryArgs;
 		FOnImageViewerOpened OnImageViewerOpened;
 		FOnGetGroupContextMenu OnGetGroupContextMenu;
 		FOnGetItemContextMenu OnGetItemContextMenu;
