Media Viewer: Added icons to the right click menu.

diff --git a/Engine/Plugins/Experimental/MediaViewer/Source/MediaViewer/Private/MediaViewerCommands.cpp b/Engine/Plugins/Experimental/MediaViewer/Source/MediaViewer/Private/MediaViewerCommands.cpp
--- a/Engine/Plugins/Experimental/MediaViewer/Source/MediaViewer/Private/MediaViewerCommands.cpp
+++ b/Engine/Plugins/Experimental/MediaViewer/Source/MediaViewer/Private/MediaViewerCommands.cpp
@@ -3,6 +3,7 @@
 #include "MediaViewerCommands.h"
 
 #include "Framework/Commands/InputChord.h"
+#include "MediaViewerStyle.h"
 #include "Styling/AppStyle.h"
 
 #define LOCTEXT_NAMESPACE "MediaViewerCommands"
@@ -11,7 +12,7 @@ namespace UE::MediaViewer::Private
 {
 
 FMediaViewerCommands::FMediaViewerCommands()
-	: TCommands("MediaViewer", LOCTEXT("ContextDescription", "Media Viewer"), NAME_None, FAppStyle::Get().GetStyleSetName())
+	: TCommands("MediaViewer", LOCTEXT("ContextDescription", "Media Viewer"), NAME_None, FMediaViewerStyle::Get().GetStyleSetName())
 {
 }
 
diff --git a/Engine/Plugins/Experimental/MediaViewer/Source/MediaViewer/Private/MediaViewerStyle.cpp b/Engine/Plugins/Experimental/MediaViewer/Source/MediaViewer/Private/MediaViewerStyle.cpp
--- a/Engine/Plugins/Experimental/MediaViewer/Source/MediaViewer/Private/MediaViewerStyle.cpp
+++ b/Engine/Plugins/Experimental/MediaViewer/Source/MediaViewer/Private/MediaViewerStyle.cpp
@@ -20,6 +20,9 @@ const FName FMediaViewerStyle::StyleName("MediaViewerStyle");
 FMediaViewerStyle::FMediaViewerStyle()
 	: FSlateStyleSet(StyleName)
 {
+	const FVector2D Image16x16(16);
+	const FVector2D Image20x20(20);
+
 	FSlateStyleSet::SetContentRoot(FPaths::EngineContentDir() / TEXT("Editor/Slate"));
 	FSlateStyleSet::SetCoreContentRoot(FPaths::EngineContentDir() / TEXT("Slate"));
 
@@ -49,16 +52,16 @@ FMediaViewerStyle::FMediaViewerStyle()
 	Set("MediaButtons", FButtonStyle(FAppStyle::Get().GetWidgetStyle<FButtonStyle>("Animation.PlayControlsButton"))
 		.SetNormal(FSlateNoResource())
 		.SetDisabled(FSlateNoResource())
-		.SetHovered(FSlateRoundedBoxBrush(FLinearColor(0.2, 0.2, 0.2, 0.5), 3.f, FVector2D(20.f)))
-		.SetPressed(FSlateRoundedBoxBrush(FLinearColor(0.1, 0.1, 0.1, 0.5), 3.f, FVector2D(20.f)))
+		.SetHovered(FSlateRoundedBoxBrush(FLinearColor(0.2, 0.2, 0.2, 0.5), 3.f, Image20x20))
+		.SetPressed(FSlateRoundedBoxBrush(FLinearColor(0.1, 0.1, 0.1, 0.5), 3.f, Image20x20))
 		.SetNormalPadding(FMargin(2.f, 2.f, 2.f, 2.f))
 		.SetPressedPadding(FMargin(2.f, 2.f, 2.f, 2.f)));
 
 	Set("LibraryButtons", FButtonStyle(FAppStyle::Get().GetWidgetStyle<FButtonStyle>("Animation.PlayControlsButton"))
 		.SetNormal(FSlateNoResource())
 		.SetDisabled(FSlateNoResource())
-		.SetHovered(FSlateRoundedBoxBrush(FLinearColor(0.2, 0.2, 0.2, 0.5), 3.f, FVector2D(20.f)))
-		.SetPressed(FSlateRoundedBoxBrush(FLinearColor(0.1, 0.1, 0.1, 0.5), 3.f, FVector2D(20.f)))
+		.SetHovered(FSlateRoundedBoxBrush(FLinearColor(0.2, 0.2, 0.2, 0.5), 3.f, Image20x20))
+		.SetPressed(FSlateRoundedBoxBrush(FLinearColor(0.1, 0.1, 0.1, 0.5), 3.f, Image20x20))
 		.SetNormalPadding(FMargin(2.f, 0.f, 2.f, 0.f))
 		.SetPressedPadding(FMargin(2.f, 0.f, 2.f, 0.f)));
 
@@ -68,6 +71,20 @@ FMediaViewerStyle::FMediaViewerStyle()
 		Set("ColorChip", new IMAGE_BRUSH_SVG(TEXT("color-chip"), FVector2D(12.0)));
 	}
 
+	// Commands
+	FSlateStyleSet::SetContentRoot(FPaths::EngineContentDir() / "Slate");
+
+	Set("MediaViewer.CopyTransform", new IMAGE_BRUSH_SVG("Starship/Common/Copy",        Image16x16));
+	Set("MediaViewer.ToggleOverlay", new IMAGE_BRUSH_SVG("Starship/Common/visible",     Image16x16));
+
+	FSlateStyleSet::SetContentRoot(FPaths::EngineContentDir() / TEXT("Editor/Slate"));
+
+	Set("MediaViewer.ResetTransform",        new IMAGE_BRUSH_SVG("Starship/Common/ResetCamera",      Image16x16));
+	Set("MediaViewer.AddToLibrary",          new IMAGE_BRUSH_SVG("Starship/Common/Unpinned",         Image16x16));
+	Set("MediaViewer.ToggleLockedTransform", new IMAGE_BRUSH(    "Common/link",                      Image16x16));
+	Set("MediaViewer.ResetAllTransforms",    new IMAGE_BRUSH_SVG("Starship/Common/ResetCamera",      Image16x16));
+	Set("MediaViewer.SwapAB",                new IMAGE_BRUSH(    "Icons/Paint/Paint_SwapColors_40x", Image16x16));
+
 	FSlateStyleRegistry::RegisterSlateStyle(*this);
 }
 
diff --git a/Engine/Plugins/Experimental/MediaViewer/Source/MediaViewer/Private/Widgets/SMediaImageViewerOverlay.cpp b/Engine/Plugins/Experimental/MediaViewer/Source/MediaViewer/Private/Widgets/SMediaImageViewerOverlay.cpp
--- a/Engine/Plugins/Experimental/MediaViewer/Source/MediaViewer/Private/Widgets/SMediaImageViewerOverlay.cpp
+++ b/Engine/Plugins/Experimental/MediaViewer/Source/MediaViewer/Private/Widgets/SMediaImageViewerOverlay.cpp
@@ -475,10 +475,15 @@ void SMediaImageViewerOverlay::BindCommands()
 
 	CommandList->MapAction(
 		Commands.ResetTransform, 
-		FExecuteAction::CreateSP(this, &SMediaImageViewerOverlay::Try_ResetTransform),
+		FExecuteAction::CreateSP(this, &SMediaImageViewerOverlay::Try_ResetTransform)
+	);
+
+	CommandList->MapAction(
+		Commands.ResetAllTransforms,
+		FExecuteAction::CreateSP(this, &SMediaImageViewerOverlay::ResetAllTransforms),
 		FCanExecuteAction(),
 		FGetActionCheckState(),
-		FIsActionButtonVisible::CreateSP(this, &SMediaImageViewerOverlay::CanResetTransform)
+		FIsActionButtonVisible::CreateSP(this, &SMediaImageViewerOverlay::CanResetAllTransforms)
 	);
 
 	CommandList->MapAction(
@@ -687,6 +692,16 @@ void SMediaImageViewerOverlay::ResetTransform()
 	ScaleToFit(/* Use transform lock */ false);
 }
 
+bool SMediaImageViewerOverlay::CanResetAllTransforms()
+{
+	return Delegates->GetActiveView.Execute() == EMediaImageViewerActivePosition::Both;
+}
+
+void SMediaImageViewerOverlay::ResetAllTransforms()
+{
+	Delegates->ResetTransformToAll.Execute();
+}
+
 FVector SMediaImageViewerOverlay::GetOffset() const
 {
 	if (TSharedPtr<FMediaImageViewer> ImageViewer = Delegates->GetImageViewer.Execute(Position))
@@ -837,11 +852,6 @@ void SMediaImageViewerOverlay::ScaleToFit(bool bInUseTransformLock)
 	}
 }
 
-bool SMediaImageViewerOverlay::CanResetTransform()
-{
-	return Delegates->GetActiveView.Execute() == EMediaImageViewerActivePosition::Both;
-}
-
 void SMediaImageViewerOverlay::Try_ResetTransform()
 {
 	if (Delegates->AreTransformsLocked.Execute())
diff --git a/Engine/Plugins/Experimental/MediaViewer/Source/MediaViewer/Private/Widgets/SMediaImageViewerOverlay.h b/Engine/Plugins/Experimental/MediaViewer/Source/MediaViewer/Private/Widgets/SMediaImageViewerOverlay.h
--- a/Engine/Plugins/Experimental/MediaViewer/Source/MediaViewer/Private/Widgets/SMediaImageViewerOverlay.h
+++ b/Engine/Plugins/Experimental/MediaViewer/Source/MediaViewer/Private/Widgets/SMediaImageViewerOverlay.h
@@ -144,8 +144,9 @@ protected:
 	void MultiplyScale(float InMultiple);
 
 	/** Interact with the image viewer's entire transform. */
-	bool CanResetTransform();
 	void ResetTransform();
+	bool CanResetAllTransforms();
+	void ResetAllTransforms();
 	void Try_ResetTransform(); // Either local or 'send to all' if transform is locked
 	void Try_SetTransform(FVector InOffset, FRotator InRotation, float InScale); // Either local or 'send to all' if transform is locked
 	bool CanCopyTransform();
