Media Viewer: Images now support yaw rotation.

diff --git a/Engine/Plugins/Experimental/MediaViewer/Source/MediaViewer/Private/ImageViewers/MediaImageViewer.cpp b/Engine/Plugins/Experimental/MediaViewer/Source/MediaViewer/Private/ImageViewers/MediaImageViewer.cpp
--- a/Engine/Plugins/Experimental/MediaViewer/Source/MediaViewer/Private/ImageViewers/MediaImageViewer.cpp
+++ b/Engine/Plugins/Experimental/MediaViewer/Source/MediaViewer/Private/ImageViewers/MediaImageViewer.cpp
@@ -196,12 +196,15 @@ void FMediaImageViewer::PaintImage(FMediaImageSlatePaintParams& InPaintParams, c
 		return;
 	}
 
-	FSlateDrawElement::MakeBox(
+	FSlateDrawElement::MakeRotatedBox(
 		InPaintParams.DrawElements,
 		InPaintParams.LayerId,
 		InPaintGeometry.Geometry,
 		Brush.Get(),
 		DrawEffects,
+		GetPaintSettings().Rotation.Yaw * 2.f * PI / 360.f,
+		{},
+		FSlateDrawElement::ERotationSpace::RelativeToElement,
 		GetPaintSettings().Tint * FLinearColor(1.f, 1.f, 1.f, InPaintParams.ImageOpacity)
 	);
 
@@ -308,21 +311,17 @@ void FMediaImageViewer::PaintImage(const FMediaImageCanvasPaintParams& InPaintPa
 
 			const float PositionX = InPaintGeometry.Position.X + (ImageUVRangeMinX * InPaintGeometry.Size.X);
 			const float SizeX = InPaintGeometry.Size.X * ImageUVRangeSize;
-			const float U = ImageUVRangeMinX * ForegroundTexture->GetSurfaceWidth();
-			const float UL = ForegroundTexture->GetSurfaceWidth() * ImageUVRangeSize;
 
-			InPaintParams.Canvas->DrawTile(
+			InPaintParams.Canvas->K2_DrawTexture(
 				ForegroundTexture,
-				PositionX,
-				InPaintGeometry.Position.Y,
-				SizeX,
-				InPaintGeometry.Size.Y,
-				U,
-				0,
-				UL,
-				ForegroundTexture->GetSurfaceHeight()
+				{PositionX, InPaintGeometry.Position.Y},
+				{SizeX, InPaintGeometry.Size.Y},
+				FVector2D(ImageUVRangeMinX, 0),
+				FVector2D(ImageUVRangeSize, 1),
+				GetPaintSettings().Tint * FLinearColor(1.f, 1.f, 1.f, InPaintParams.ImageOpacity),
+				BLEND_Translucent,
+				GetPaintSettings().Rotation.Yaw
 			);
-
 			break;
 		}
 
@@ -350,19 +349,16 @@ void FMediaImageViewer::PaintImage(const FMediaImageCanvasPaintParams& InPaintPa
 
 			const float PositionY = InPaintGeometry.Position.Y + (ImageUVRangeMinY * InPaintGeometry.Size.Y);
 			const float SizeY = InPaintGeometry.Size.Y * ImageUVRangeSize;
-			const float V = ImageUVRangeMinY * ForegroundTexture->GetSurfaceHeight();
-			const float VL = ForegroundTexture->GetSurfaceHeight() * ImageUVRangeSize;
 
-			InPaintParams.Canvas->DrawTile(
+			InPaintParams.Canvas->K2_DrawTexture(
 				ForegroundTexture,
-				InPaintGeometry.Position.X,
-				PositionY,
-				InPaintGeometry.Size.X,
-				SizeY,
-				0,
-				V,
-				ForegroundTexture->GetSurfaceWidth(),
-				VL
+				{InPaintGeometry.Position.X, PositionY},
+				{InPaintGeometry.Size.X, SizeY},
+				FVector2D(0, ImageUVRangeMinY),
+				FVector2D(1, ImageUVRangeSize),
+				GetPaintSettings().Tint * FLinearColor(1.f, 1.f, 1.f, InPaintParams.ImageOpacity),
+				BLEND_Translucent,
+				GetPaintSettings().Rotation.Yaw
 			);
 
 			break;
diff --git a/Engine/Plugins/Experimental/MediaViewer/Source/MediaViewer/Private/Library/MediaViewerLibrary.h b/Engine/Plugins/Experimental/MediaViewer/Source/MediaViewer/Private/Library/MediaViewerLibrary.h
--- a/Engine/Plugins/Experimental/MediaViewer/Source/MediaViewer/Private/Library/MediaViewerLibrary.h
+++ b/Engine/Plugins/Experimental/MediaViewer/Source/MediaViewer/Private/Library/MediaViewerLibrary.h
@@ -18,7 +18,7 @@ namespace UE::MediaViewer::Private
 constexpr int32 MaxHistoryEntries = 20;
 
 /**
- * Implementation of @see FMediaViewerLibrary.
+ * Implementation of @see IMediaViewerLibrary.
  */
 class FMediaViewerLibrary : public IMediaViewerLibrary, public TSharedFromThis<FMediaViewerLibrary>
 {
diff --git a/Engine/Plugins/Experimental/MediaViewer/Source/MediaViewer/Public/IMediaViewerModule.h b/Engine/Plugins/Experimental/MediaViewer/Source/MediaViewer/Public/IMediaViewerModule.h
--- a/Engine/Plugins/Experimental/MediaViewer/Source/MediaViewer/Public/IMediaViewerModule.h
+++ b/Engine/Plugins/Experimental/MediaViewer/Source/MediaViewer/Public/IMediaViewerModule.h
@@ -55,7 +55,7 @@ public:
 	/** If the settings differ from the currently opened tab's settings, it will regenerate the entire display. */
 	virtual bool OpenTab(const FMediaViewerArgs& InMediaViewerArgs) = 0;
 
-	/** Gets the Library for the open tab. */
+	/** Returns the Library, which is common between all tabs. */
 	virtual TSharedRef<IMediaViewerLibrary> GetLibrary() const = 0;
 
 	virtual TSharedPtr<FMediaViewerLibraryItem> CreateLibraryItem(const FAssetData& InAssetData) const = 0;
