Media Viewer: Added bookmark system

diff --git a/Engine/Plugins/Experimental/MediaViewer/Source/MediaViewer/Private/Library/MediaViewerLibraryIni.cpp b/Engine/Plugins/Experimental/MediaViewer/Source/MediaViewer/Private/Library/MediaViewerLibraryIni.cpp
--- a/Engine/Plugins/Experimental/MediaViewer/Source/MediaViewer/Private/Library/MediaViewerLibraryIni.cpp
+++ b/Engine/Plugins/Experimental/MediaViewer/Source/MediaViewer/Private/Library/MediaViewerLibraryIni.cpp
@@ -15,6 +15,11 @@ UMediaViewerLibraryIni& UMediaViewerLibraryIni::Get()
 	return *GetMutableDefault<UMediaViewerLibraryIni>();
 }
 
+UMediaViewerLibraryIni::UMediaViewerLibraryIni()
+{
+	Bookmarks.SetNum(10);
+}
+
 void UMediaViewerLibraryIni::SaveLibrary(const TSharedRef<UE::MediaViewer::IMediaViewerLibrary>& InLibrary)
 {
 	using namespace UE::MediaViewer::Private;
@@ -123,12 +128,25 @@ bool UMediaViewerLibraryIni::HasItem(const FGuid& InItemId) const
 	return false;
 }
 
-TConstArrayView<FMediaViewerState> UMediaViewerLibraryIni::GetSavedStates() const
+TConstArrayView<FMediaViewerState> UMediaViewerLibraryIni::GetBookmarks() const
+{
+	return Bookmarks;
+}
+
+void UMediaViewerLibraryIni::SetBookmark(int32 InIndex, const FMediaViewerState& InState)
+{
+	if (Bookmarks.IsValidIndex(InIndex))
+	{
+		Bookmarks[InIndex] = InState;
+	}
+}
+
+TOptional<FMediaViewerState> UMediaViewerLibraryIni::GetLastOpenedState() const
 {
-	return SavedStates;
+	return LastOpenedState;
 }
 
-void UMediaViewerLibraryIni::SetSavedStates(const TArray<FMediaViewerState>& InStates)
+void UMediaViewerLibraryIni::SetLastOpenedState(const FMediaViewerState& InState)
 {
-	SavedStates = InStates;
+	LastOpenedState = InState;
 }
diff --git a/Engine/Plugins/Experimental/MediaViewer/Source/MediaViewer/Private/Library/MediaViewerLibraryIni.h b/Engine/Plugins/Experimental/MediaViewer/Source/MediaViewer/Private/Library/MediaViewerLibraryIni.h
--- a/Engine/Plugins/Experimental/MediaViewer/Source/MediaViewer/Private/Library/MediaViewerLibraryIni.h
+++ b/Engine/Plugins/Experimental/MediaViewer/Source/MediaViewer/Private/Library/MediaViewerLibraryIni.h
@@ -8,6 +8,7 @@
 #include "Library/MediaViewerLibraryItem.h"
 #include "ImageViewer/MediaImageViewer.h"
 #include "Misc/Guid.h"
+#include "Misc/Optional.h"
 #include "Widgets/MediaViewerSettings.h"
 #include "Widgets/SMediaViewer.h"
 
@@ -71,6 +72,8 @@ class UMediaViewerLibraryIni : public UObject
 public:
 	static UMediaViewerLibraryIni& Get();
 
+	UMediaViewerLibraryIni();
+
 	void SaveLibrary(const TSharedRef<UE::MediaViewer::IMediaViewerLibrary>& InLibrary);
 
 	void LoadLibrary(const TSharedRef<UE::MediaViewer::IMediaViewerLibrary>& InLibrary) const;
@@ -79,9 +82,13 @@ public:
 
 	bool HasItem(const FGuid& InItemId) const;
 
-	TConstArrayView<FMediaViewerState> GetSavedStates() const;
+	TConstArrayView<FMediaViewerState> GetBookmarks() const;
+
+	void SetBookmark(int32 InIndex, const FMediaViewerState& InState);
+
+	TOptional<FMediaViewerState> GetLastOpenedState() const;
 
-	void SetSavedStates(const TArray<FMediaViewerState>& InStates);
+	void SetLastOpenedState(const FMediaViewerState& InState);
 
 protected:
 	UPROPERTY(Config)
@@ -91,5 +98,8 @@ protected:
 	TArray<FMediaViewerLibraryItemData> Items;
 
 	UPROPERTY(Config)
-	TArray<FMediaViewerState> SavedStates;
+	TArray<FMediaViewerState> Bookmarks;
+
+	UPROPERTY(Config)
+	TOptional<FMediaViewerState> LastOpenedState;
 };
diff --git a/Engine/Plugins/Experimental/MediaViewer/Source/MediaViewer/Private/MediaViewerCommands.cpp b/Engine/Plugins/Experimental/MediaViewer/Source/MediaViewer/Private/MediaViewerCommands.cpp
--- a/Engine/Plugins/Experimental/MediaViewer/Source/MediaViewer/Private/MediaViewerCommands.cpp
+++ b/Engine/Plugins/Experimental/MediaViewer/Source/MediaViewer/Private/MediaViewerCommands.cpp
@@ -129,11 +129,35 @@ void FMediaViewerCommands::RegisterCommands()
 	);
 
 	UI_COMMAND(
-		Scale12, 
-		"12.5%", 
-		"Set Scale to 12.5%.", 
-		EUserInterfaceActionType::RadioButton, 
-		FInputChord(EKeys::One)
+		ScaleUp, 
+		"Zoom In", 
+		"Zoom in on the image.", 
+		EUserInterfaceActionType::Button, 
+		FInputChord(EKeys::Equals)
+	);
+
+	UI_COMMAND(
+		ScaleDown,
+		"Zoom Out",
+		"Zoom out of the image.",
+		EUserInterfaceActionType::Button,
+		FInputChord(EKeys::Hyphen)
+	);
+
+	UI_COMMAND(
+		ScaleToFit, 
+		"Scale to Fit", 
+		"Changes the scale of the image so it matches, but does not exceed, the size of the viewport.", 
+		EUserInterfaceActionType::Button,
+		FInputChord(EKeys::F)
+	);
+
+	UI_COMMAND(
+		Scale12,
+		"12.5%",
+		"Set Scale to 12.5%.",
+		EUserInterfaceActionType::RadioButton,
+		FInputChord()
 	);
 
 	UI_COMMAND(
@@ -141,7 +165,7 @@ void FMediaViewerCommands::RegisterCommands()
 		"25%",
 		"Set Scale to 25%.",
 		EUserInterfaceActionType::RadioButton,
-		FInputChord(EKeys::Two)
+		FInputChord()
 	);
 
 	UI_COMMAND(
@@ -149,7 +173,7 @@ void FMediaViewerCommands::RegisterCommands()
 		"50%",
 		"Set Scale to 50%.",
 		EUserInterfaceActionType::RadioButton,
-		FInputChord(EKeys::Three)
+		FInputChord()
 	);
 
 	UI_COMMAND(
@@ -157,7 +181,7 @@ void FMediaViewerCommands::RegisterCommands()
 		"100%",
 		"Set Scale to 100%.",
 		EUserInterfaceActionType::RadioButton,
-		FInputChord(EKeys::Four)
+		FInputChord()
 	);
 
 	UI_COMMAND(
@@ -165,7 +189,7 @@ void FMediaViewerCommands::RegisterCommands()
 		"200%",
 		"Set Scale to 200%.",
 		EUserInterfaceActionType::RadioButton,
-		FInputChord(EKeys::Five)
+		FInputChord()
 	);
 
 	UI_COMMAND(
@@ -173,7 +197,7 @@ void FMediaViewerCommands::RegisterCommands()
 		"400%",
 		"Set Scale to 400%.",
 		EUserInterfaceActionType::RadioButton,
-		FInputChord(EKeys::Six)
+		FInputChord()
 	);
 
 	UI_COMMAND(
@@ -181,15 +205,7 @@ void FMediaViewerCommands::RegisterCommands()
 		"800%",
 		"Set Scale to 800%.",
 		EUserInterfaceActionType::RadioButton,
-		FInputChord(EKeys::Seven)
-	);
-
-	UI_COMMAND(
-		ScaleToFit, 
-		"Scale to Fit", 
-		"Changes the scale of the image so it matches, but does not exceed, the size of the viewport.", 
-		EUserInterfaceActionType::Button,
-		FInputChord(EKeys::Eight)
+		FInputChord()
 	);
 
 	UI_COMMAND(
@@ -226,22 +242,6 @@ void FMediaViewerCommands::RegisterCommands()
 		FInputChord(EKeys::T)
 	);
 
-	UI_COMMAND(
-		MipMinus, 
-		"Lower Mip", 
-		"Show next lower mip.", 
-		EUserInterfaceActionType::Button,
-		FInputChord(EKeys::Hyphen)
-	);
-
-	UI_COMMAND(
-		MipPlus, 
-		"Higher Mip", 
-		"Show next higher mip.", 
-		EUserInterfaceActionType::Button,
-		FInputChord(EKeys::Equals)
-	);
-
 	UI_COMMAND(
 		SecondImageOpacityMinus,
 		"Lower B Image Opacity",
@@ -357,6 +357,166 @@ void FMediaViewerCommands::RegisterCommands()
 		FInputChord(EKeys::N, EModifierKey::Control),
 		FInputChord(EKeys::M, EModifierKey::Command)
 	);
+
+	UI_COMMAND(
+		OpenBookmark1,
+		"Bookmark 1",
+		"Restore bookmark 1.",
+		EUserInterfaceActionType::Button,
+		FInputChord(EKeys::One)
+	);
+
+	UI_COMMAND(
+		OpenBookmark2,
+		"Bookmark 2",
+		"Restore bookmark 2.",
+		EUserInterfaceActionType::Button,
+		FInputChord(EKeys::Two)
+	);
+
+	UI_COMMAND(
+		OpenBookmark3,
+		"Bookmark 3",
+		"Restore bookmark 3.",
+		EUserInterfaceActionType::Button,
+		FInputChord(EKeys::Three)
+	);
+
+	UI_COMMAND(
+		OpenBookmark4,
+		"Bookmark 4",
+		"Restore bookmark 4.",
+		EUserInterfaceActionType::Button,
+		FInputChord(EKeys::Four)
+	);
+
+	UI_COMMAND(
+		OpenBookmark5,
+		"Bookmark 5",
+		"Restore bookmark 5.",
+		EUserInterfaceActionType::Button,
+		FInputChord(EKeys::Five)
+	);
+
+	UI_COMMAND(
+		OpenBookmark6,
+		"Bookmark 6",
+		"Restore bookmark 6.",
+		EUserInterfaceActionType::Button,
+		FInputChord(EKeys::Six)
+	);
+
+	UI_COMMAND(
+		OpenBookmark7,
+		"Bookmark 7",
+		"Restore bookmark 7.",
+		EUserInterfaceActionType::Button,
+		FInputChord(EKeys::Seven)
+	);
+
+	UI_COMMAND(
+		OpenBookmark8,
+		"Bookmark 8",
+		"Restore bookmark 8.",
+		EUserInterfaceActionType::Button,
+		FInputChord(EKeys::Eight)
+	);
+
+	UI_COMMAND(
+		OpenBookmark9,
+		"Bookmark 9",
+		"Restore bookmark 9.",
+		EUserInterfaceActionType::Button,
+		FInputChord(EKeys::Nine)
+	);
+
+	UI_COMMAND(
+		OpenBookmark10,
+		"Bookmark 10",
+		"Restore bookmark 10.",
+		EUserInterfaceActionType::Button,
+		FInputChord(EKeys::Zero)
+	);
+
+	UI_COMMAND(
+		SaveBookmark1,
+		"Save Bookmark 1",
+		"Save bookmark 1.",
+		EUserInterfaceActionType::Button,
+		FInputChord(EKeys::One, EModifierKey::Control)
+	);
+
+	UI_COMMAND(
+		SaveBookmark2,
+		"Save Bookmark 2",
+		"Save bookmark 2.",
+		EUserInterfaceActionType::Button,
+		FInputChord(EKeys::Two, EModifierKey::Control)
+	);
+
+	UI_COMMAND(
+		SaveBookmark3,
+		"Save Bookmark 3",
+		"Save bookmark 3.",
+		EUserInterfaceActionType::Button,
+		FInputChord(EKeys::Three, EModifierKey::Control)
+	);
+
+	UI_COMMAND(
+		SaveBookmark4,
+		"Save Bookmark 4",
+		"Save bookmark 4.",
+		EUserInterfaceActionType::Button,
+		FInputChord(EKeys::Four, EModifierKey::Control)
+	);
+
+	UI_COMMAND(
+		SaveBookmark5,
+		"Save Bookmark 5",
+		"Save bookmark 5.",
+		EUserInterfaceActionType::Button,
+		FInputChord(EKeys::Five, EModifierKey::Control)
+	);
+
+	UI_COMMAND(
+		SaveBookmark6,
+		"Save Bookmark 6",
+		"Save bookmark 6.",
+		EUserInterfaceActionType::Button,
+		FInputChord(EKeys::Six, EModifierKey::Control)
+	);
+
+	UI_COMMAND(
+		SaveBookmark7,
+		"Save Bookmark 7",
+		"Save bookmark 7.",
+		EUserInterfaceActionType::Button,
+		FInputChord(EKeys::Seven, EModifierKey::Control)
+	);
+
+	UI_COMMAND(
+		SaveBookmark8,
+		"Save Bookmark 8",
+		"Save bookmark 8.",
+		EUserInterfaceActionType::Button,
+		FInputChord(EKeys::Eight, EModifierKey::Control)
+	);
+
+	UI_COMMAND(
+		SaveBookmark9,
+		"Save Bookmark 9",
+		"Save bookmark 9.",
+		EUserInterfaceActionType::Button,
+		FInputChord(EKeys::Nine, EModifierKey::Control)
+	);
+
+	UI_COMMAND(
+		SaveBookmark10,
+		"Save Bookmark 10",
+		"Save bookmark 10.",
+		EUserInterfaceActionType::Button,
+		FInputChord(EKeys::Zero, EModifierKey::Control)
+	);
 }
 
 } // UE::MediaViewer::Private
diff --git a/Engine/Plugins/Experimental/MediaViewer/Source/MediaViewer/Private/MediaViewerCommands.h b/Engine/Plugins/Experimental/MediaViewer/Source/MediaViewer/Private/MediaViewerCommands.h
--- a/Engine/Plugins/Experimental/MediaViewer/Source/MediaViewer/Private/MediaViewerCommands.h
+++ b/Engine/Plugins/Experimental/MediaViewer/Source/MediaViewer/Private/MediaViewerCommands.h
@@ -32,6 +32,9 @@ public:
 	TSharedPtr<FUICommandInfo> RotatePlusRoll;
 	TSharedPtr<FUICommandInfo> RotateMinusRoll;
 	
+	TSharedPtr<FUICommandInfo> ScaleUp;
+	TSharedPtr<FUICommandInfo> ScaleDown;
+
 	TSharedPtr<FUICommandInfo> Scale12;
 	TSharedPtr<FUICommandInfo> Scale25;
 	TSharedPtr<FUICommandInfo> Scale50;
@@ -39,7 +42,6 @@ public:
 	TSharedPtr<FUICommandInfo> Scale200;
 	TSharedPtr<FUICommandInfo> Scale400;
 	TSharedPtr<FUICommandInfo> Scale800;
-
 	TSharedPtr<FUICommandInfo> ScaleToFit;
 
 	TSharedPtr<FUICommandInfo> CopyTransform;
@@ -47,9 +49,6 @@ public:
 	TSharedPtr<FUICommandInfo> ResetAllTransforms;
 	TSharedPtr<FUICommandInfo> ToggleLockedTransform;
 
-	TSharedPtr<FUICommandInfo> MipMinus;
-	TSharedPtr<FUICommandInfo> MipPlus;
-
 	TSharedPtr<FUICommandInfo> SecondImageOpacityMinus;
 	TSharedPtr<FUICommandInfo> SecondImageOpacityPlus;
 
@@ -70,6 +69,28 @@ public:
 	TSharedPtr<FUICommandInfo> QuickFastForwardVideo;
 
 	TSharedPtr<FUICommandInfo> OpenNewTab;
+
+	TSharedPtr<FUICommandInfo> OpenBookmark1;
+	TSharedPtr<FUICommandInfo> OpenBookmark2;
+	TSharedPtr<FUICommandInfo> OpenBookmark3;
+	TSharedPtr<FUICommandInfo> OpenBookmark4;
+	TSharedPtr<FUICommandInfo> OpenBookmark5;
+	TSharedPtr<FUICommandInfo> OpenBookmark6;
+	TSharedPtr<FUICommandInfo> OpenBookmark7;
+	TSharedPtr<FUICommandInfo> OpenBookmark8;
+	TSharedPtr<FUICommandInfo> OpenBookmark9;
+	TSharedPtr<FUICommandInfo> OpenBookmark10;
+
+	TSharedPtr<FUICommandInfo> SaveBookmark1;
+	TSharedPtr<FUICommandInfo> SaveBookmark2;
+	TSharedPtr<FUICommandInfo> SaveBookmark3;
+	TSharedPtr<FUICommandInfo> SaveBookmark4;
+	TSharedPtr<FUICommandInfo> SaveBookmark5;
+	TSharedPtr<FUICommandInfo> SaveBookmark6;
+	TSharedPtr<FUICommandInfo> SaveBookmark7;
+	TSharedPtr<FUICommandInfo> SaveBookmark8;
+	TSharedPtr<FUICommandInfo> SaveBookmark9;
+	TSharedPtr<FUICommandInfo> SaveBookmark10;
 };
 
 } // UE::MediaViewer::Private
diff --git a/Engine/Plugins/Experimental/MediaViewer/Source/MediaViewer/Private/MediaViewerModule.cpp b/Engine/Plugins/Experimental/MediaViewer/Source/MediaViewer/Private/MediaViewerModule.cpp
--- a/Engine/Plugins/Experimental/MediaViewer/Source/MediaViewer/Private/MediaViewerModule.cpp
+++ b/Engine/Plugins/Experimental/MediaViewer/Source/MediaViewer/Private/MediaViewerModule.cpp
@@ -405,7 +405,7 @@ void FMediaViewerModule::OnTabClosed(TSharedRef<SDockTab> InDockTab)
 	if (Content->GetWidgetClass().GetWidgetType() == SMediaViewerTab::StaticWidgetClass().GetWidgetType())
 	{
 		TSharedRef<SMediaViewerTab> MediaViewerTab = StaticCastSharedRef<SMediaViewerTab>(Content);
-		MediaViewerTab->GetViewer()->SaveStates(/* Save config */ false);
+		MediaViewerTab->GetViewer()->SaveLastOpenedState(/* Save config */ false);
 
 		UMediaViewerLibraryIni& Ini = UMediaViewerLibraryIni::Get();
 		Ini.SaveLibrary(MediaViewerTab->GetLibrary());
diff --git a/Engine/Plugins/Experimental/MediaViewer/Source/MediaViewer/Private/Widgets/SMediaImageViewerOverlay.cpp b/Engine/Plugins/Experimental/MediaViewer/Source/MediaViewer/Private/Widgets/SMediaImageViewerOverlay.cpp
--- a/Engine/Plugins/Experimental/MediaViewer/Source/MediaViewer/Private/Widgets/SMediaImageViewerOverlay.cpp
+++ b/Engine/Plugins/Experimental/MediaViewer/Source/MediaViewer/Private/Widgets/SMediaImageViewerOverlay.cpp
@@ -259,17 +259,7 @@ FReply SMediaImageViewerOverlay::OnMouseButtonUp(const FGeometry& InMyGeometry,
 
 FReply SMediaImageViewerOverlay::OnMouseWheel(const FGeometry& InMyGeometry, const FPointerEvent& InMouseEvent)
 {
-	static const float NormalMultiplier = FMath::Pow(2.0, 0.125);
-	static const float FastMultiplier = FMath::Pow(2.0, 0.5);
-	static const float SlowMultiplier = FMath::Pow(2.0, 0.03125);
-
-	const FModifierKeysState ModifierKeys = FSlateApplication::Get().GetModifierKeys();
-
-	const float ModifierScale = ModifierKeys.IsLeftShiftDown()
-		? FastMultiplier
-		: ((ModifierKeys.IsLeftControlDown() || ModifierKeys.IsLeftCommandDown())
-			? SlowMultiplier
-			: NormalMultiplier);
+	const float ModifierScale = GetScaleModifier();
 
 	const float ScaleMultiplier = InMouseEvent.GetWheelDelta() > 0
 		? ModifierScale
@@ -460,6 +450,14 @@ void SMediaImageViewerOverlay::BindCommands()
 	CommandList->MapAction(Commands.RotateMinusRoll,  FExecuteAction::CreateSP(this, &SMediaImageViewerOverlay::Try_AddRotation, FRotator(  0,   0, -10)));
 	CommandList->MapAction(Commands.RotatePlusRoll,   FExecuteAction::CreateSP(this, &SMediaImageViewerOverlay::Try_AddRotation, FRotator(  0,   0,  10)));
 
+	CommandList->MapAction(Commands.ScaleUp, FExecuteAction::CreateSP(this, &SMediaImageViewerOverlay::ScaleUp));
+	CommandList->MapAction(Commands.ScaleDown, FExecuteAction::CreateSP(this, &SMediaImageViewerOverlay::ScaleDown));
+
+	CommandList->MapAction(
+		Commands.ScaleToFit, 
+		FExecuteAction::CreateSP(this, &SMediaImageViewerOverlay::ScaleToFit, /* Use transform lock */ true)
+	);
+
 	CommandList->MapAction(Commands.Scale12,  FExecuteAction::CreateSP(this, &SMediaImageViewerOverlay::Try_SetScale, 0.125f));
 	CommandList->MapAction(Commands.Scale25,  FExecuteAction::CreateSP(this, &SMediaImageViewerOverlay::Try_SetScale, 0.25f));
 	CommandList->MapAction(Commands.Scale50,  FExecuteAction::CreateSP(this, &SMediaImageViewerOverlay::Try_SetScale, 0.5f));
@@ -468,11 +466,6 @@ void SMediaImageViewerOverlay::BindCommands()
 	CommandList->MapAction(Commands.Scale400, FExecuteAction::CreateSP(this, &SMediaImageViewerOverlay::Try_SetScale, 4.0f));
 	CommandList->MapAction(Commands.Scale800, FExecuteAction::CreateSP(this, &SMediaImageViewerOverlay::Try_SetScale, 8.0f));
 
-	CommandList->MapAction(
-		Commands.ScaleToFit, 
-		FExecuteAction::CreateSP(this, &SMediaImageViewerOverlay::ScaleToFit, /* Use transform lock */ true)
-	);
-
 	CommandList->MapAction(
 		Commands.ResetTransform, 
 		FExecuteAction::CreateSP(this, &SMediaImageViewerOverlay::Try_ResetTransform)
@@ -494,9 +487,6 @@ void SMediaImageViewerOverlay::BindCommands()
 		FIsActionButtonVisible::CreateSP(this, &SMediaImageViewerOverlay::CanCopyTransform)
 	);
 
-	CommandList->MapAction(Commands.MipMinus, FExecuteAction::CreateSP(this, &SMediaImageViewerOverlay::AdjustMipLevel, static_cast<int8>(-1)));
-	CommandList->MapAction(Commands.MipPlus,  FExecuteAction::CreateSP(this, &SMediaImageViewerOverlay::AdjustMipLevel, static_cast<int8>(1)));
-
 	CommandList->MapAction(Commands.CopyColor, FExecuteAction::CreateSP(this, &SMediaImageViewerOverlay::CopyColor));
 
 	CommandList->MapAction(
@@ -802,6 +792,47 @@ void SMediaImageViewerOverlay::SetScale(float InScale)
 	}
 }
 
+float SMediaImageViewerOverlay::GetScaleModifier()
+{
+	static const float NormalMultiplier = FMath::Pow(2.0, 0.125);
+	static const float FastMultiplier = FMath::Pow(2.0, 0.5);
+	static const float SlowMultiplier = FMath::Pow(2.0, 0.03125);
+
+	const FModifierKeysState ModifierKeys = FSlateApplication::Get().GetModifierKeys();
+
+	return ModifierKeys.IsLeftShiftDown()
+		? FastMultiplier
+		: ((ModifierKeys.IsLeftControlDown() || ModifierKeys.IsLeftCommandDown())
+			? SlowMultiplier
+			: NormalMultiplier);
+}
+
+void SMediaImageViewerOverlay::ScaleUp()
+{
+	const float ModifierScale = GetScaleModifier();
+
+	Try_MultipyScale(GetScaleModifier());
+}
+
+void SMediaImageViewerOverlay::ScaleDown()
+{
+	const float ModifierScale = GetScaleModifier();
+
+	Try_MultipyScale(1.f / GetScaleModifier());
+}
+
+void SMediaImageViewerOverlay::Try_MultipyScale(float InMultiple)
+{
+	if (Delegates->AreTransformsLocked.Execute())
+	{
+		Delegates->MultiplyScaleToAll.Execute(InMultiple);
+	}
+	else
+	{
+		SetScale(GetScale() * InMultiple);
+	}
+}
+
 void SMediaImageViewerOverlay::MultiplyScale(float InMultiple)
 {
 	SetScale(GetScale() * InMultiple);
@@ -990,13 +1021,15 @@ void SMediaImageViewerOverlay::AddToLibrary()
 
 TSharedRef<SWidget> SMediaImageViewerOverlay::CreateMenu()
 {
+	FMediaViewerCommands& Commands = FMediaViewerCommands::Get();
+
 	FMenuBuilder MenuBuilder(true, CommandList, nullptr, false, &FAppStyle::Get(), false);
 
 	MenuBuilder.BeginSection(TEXT("Image"));
 	{
-		MenuBuilder.AddMenuEntry(FMediaViewerCommands::Get().ResetTransform);
-		MenuBuilder.AddMenuEntry(FMediaViewerCommands::Get().CopyTransform);
-		MenuBuilder.AddMenuEntry(FMediaViewerCommands::Get().AddToLibrary);
+		MenuBuilder.AddMenuEntry(Commands.ResetTransform);
+		MenuBuilder.AddMenuEntry(Commands.CopyTransform);
+		MenuBuilder.AddMenuEntry(Commands.AddToLibrary);
 	}
 	MenuBuilder.EndSection();
 
@@ -1004,10 +1037,70 @@ TSharedRef<SWidget> SMediaImageViewerOverlay::CreateMenu()
 
 	MenuBuilder.BeginSection(TEXT("Media Viewer"));
 	{
-		MenuBuilder.AddMenuEntry(FMediaViewerCommands::Get().ToggleOverlay);
-		MenuBuilder.AddMenuEntry(FMediaViewerCommands::Get().ToggleLockedTransform);
-		MenuBuilder.AddMenuEntry(FMediaViewerCommands::Get().ResetAllTransforms);
-		MenuBuilder.AddMenuEntry(FMediaViewerCommands::Get().SwapAB);
+		MenuBuilder.AddMenuEntry(Commands.ToggleOverlay);
+		MenuBuilder.AddMenuEntry(Commands.ToggleLockedTransform);
+		MenuBuilder.AddMenuEntry(Commands.ResetAllTransforms);
+		MenuBuilder.AddMenuEntry(Commands.SwapAB);
+	}
+
+	FNewMenuDelegate LoadBookmarksMenu = FNewMenuDelegate::CreateLambda(
+		[](FMenuBuilder& InMenuBuilder)
+		{
+			FMediaViewerCommands& Commands = FMediaViewerCommands::Get();
+
+			InMenuBuilder.BeginSection("Bookmarks");
+			{
+				InMenuBuilder.AddMenuEntry(Commands.OpenBookmark1);
+				InMenuBuilder.AddMenuEntry(Commands.OpenBookmark2);
+				InMenuBuilder.AddMenuEntry(Commands.OpenBookmark3);
+				InMenuBuilder.AddMenuEntry(Commands.OpenBookmark4);
+				InMenuBuilder.AddMenuEntry(Commands.OpenBookmark5);
+				InMenuBuilder.AddMenuEntry(Commands.OpenBookmark6);
+				InMenuBuilder.AddMenuEntry(Commands.OpenBookmark7);
+				InMenuBuilder.AddMenuEntry(Commands.OpenBookmark8);
+				InMenuBuilder.AddMenuEntry(Commands.OpenBookmark9);
+				InMenuBuilder.AddMenuEntry(Commands.OpenBookmark10);
+			}
+			InMenuBuilder.EndSection();
+		}
+	);
+
+	FNewMenuDelegate SaveBookmarksMenu = FNewMenuDelegate::CreateLambda(
+		[](FMenuBuilder& InMenuBuilder)
+		{
+			FMediaViewerCommands& Commands = FMediaViewerCommands::Get();
+
+			InMenuBuilder.BeginSection("Bookmarks");
+			{
+				InMenuBuilder.AddMenuEntry(Commands.SaveBookmark1);
+				InMenuBuilder.AddMenuEntry(Commands.SaveBookmark2);
+				InMenuBuilder.AddMenuEntry(Commands.SaveBookmark3);
+				InMenuBuilder.AddMenuEntry(Commands.SaveBookmark4);
+				InMenuBuilder.AddMenuEntry(Commands.SaveBookmark5);
+				InMenuBuilder.AddMenuEntry(Commands.SaveBookmark6);
+				InMenuBuilder.AddMenuEntry(Commands.SaveBookmark7);
+				InMenuBuilder.AddMenuEntry(Commands.SaveBookmark8);
+				InMenuBuilder.AddMenuEntry(Commands.SaveBookmark9);
+				InMenuBuilder.AddMenuEntry(Commands.SaveBookmark10);
+			}
+			InMenuBuilder.EndSection();
+		}
+	);
+
+	MenuBuilder.AddMenuSeparator();
+	MenuBuilder.BeginSection(TEXT("Bookmarks"));
+	{
+		MenuBuilder.AddSubMenu(
+			LOCTEXT("OpenBookmarkMenu", "Open Bookmark"),
+			LOCTEXT("OpenBookmarkMenuTooltip", "Load a saved bookmark."),
+			LoadBookmarksMenu
+		);
+
+		MenuBuilder.AddSubMenu(
+			LOCTEXT("SaveBookmarkMenu", "Save Bookmark"),
+			LOCTEXT("SaveBookmarkMenuTooltip", "Save a bookmark of the current state."),
+			LoadBookmarksMenu
+		);
 	}
 
 	return MenuBuilder.MakeWidget();
diff --git a/Engine/Plugins/Experimental/MediaViewer/Source/MediaViewer/Private/Widgets/SMediaImageViewerOverlay.h b/Engine/Plugins/Experimental/MediaViewer/Source/MediaViewer/Private/Widgets/SMediaImageViewerOverlay.h
--- a/Engine/Plugins/Experimental/MediaViewer/Source/MediaViewer/Private/Widgets/SMediaImageViewerOverlay.h
+++ b/Engine/Plugins/Experimental/MediaViewer/Source/MediaViewer/Private/Widgets/SMediaImageViewerOverlay.h
@@ -141,6 +141,10 @@ protected:
 	float GetScale() const;
 	void Try_SetScale(float InScale); // Either local or 'send to all' if transform is locked
 	void SetScale(float InScale);
+	float GetScaleModifier();
+	void ScaleUp();
+	void ScaleDown();
+	void Try_MultipyScale(float InMultiple);
 	void MultiplyScale(float InMultiple);
 
 	/** Interact with the image viewer's entire transform. */
diff --git a/Engine/Plugins/Experimental/MediaViewer/Source/MediaViewer/Private/Widgets/SMediaViewer.cpp b/Engine/Plugins/Experimental/MediaViewer/Source/MediaViewer/Private/Widgets/SMediaViewer.cpp
--- a/Engine/Plugins/Experimental/MediaViewer/Source/MediaViewer/Private/Widgets/SMediaViewer.cpp
+++ b/Engine/Plugins/Experimental/MediaViewer/Source/MediaViewer/Private/Widgets/SMediaViewer.cpp
@@ -242,6 +242,32 @@ void SMediaViewer::BindCommands()
 	);
 
 	CommandList->MapAction(Commands.OpenNewTab, FExecuteAction::CreateSP(this, &SMediaViewer::OpenNewTab));
+
+	auto BindBookmark = [this](const TSharedPtr<FUICommandInfo>& InOpenCommand, const TSharedPtr<FUICommandInfo>& InSaveCommand, int32 InIndex)
+		{
+			CommandList->MapAction(
+				InOpenCommand,
+				FExecuteAction::CreateSP(this, &SMediaViewer::OpenBookmark, InIndex),
+				FCanExecuteAction::CreateSP(this, &SMediaViewer::CanOpenBookmark, InIndex)
+			);
+
+			CommandList->MapAction(
+				InSaveCommand,
+				FExecuteAction::CreateSP(this, &SMediaViewer::SaveBookmark, InIndex),
+				FCanExecuteAction::CreateSP(this, &SMediaViewer::CanSaveBookmark)
+			);
+		};
+
+	BindBookmark(Commands.OpenBookmark1, Commands.SaveBookmark1, 1);
+	BindBookmark(Commands.OpenBookmark2, Commands.SaveBookmark2, 2);
+	BindBookmark(Commands.OpenBookmark3, Commands.SaveBookmark3, 3);
+	BindBookmark(Commands.OpenBookmark4, Commands.SaveBookmark4, 4);
+	BindBookmark(Commands.OpenBookmark5, Commands.SaveBookmark5, 5);
+	BindBookmark(Commands.OpenBookmark6, Commands.SaveBookmark6, 6);
+	BindBookmark(Commands.OpenBookmark7, Commands.SaveBookmark7, 7);
+	BindBookmark(Commands.OpenBookmark8, Commands.SaveBookmark8, 8);
+	BindBookmark(Commands.OpenBookmark9, Commands.SaveBookmark9, 9);
+	BindBookmark(Commands.OpenBookmark10, Commands.SaveBookmark10, 10);
 }
 
 bool SMediaViewer::IsShowingBothImageViewers() const
@@ -322,7 +348,7 @@ void SMediaViewer::Tick(const FGeometry& AllottedGeometry, const double InCurren
 
 		if (TickCount == TickCheckValue)
 		{
-			CheckLoadState();
+			CheckLastOpenedState();
 		}
 	}
 
@@ -757,7 +783,7 @@ TSharedRef<FMediaViewerDelegates> SMediaViewer::GetDelegates() const
 	return Delegates.ToSharedRef();
 }
 
-void SMediaViewer::SaveStates(bool bInSaveConfig) const
+FMediaViewerState SMediaViewer::GetState() const
 {
 	FMediaViewerState State;
 	State.ViewerSettings = MediaViewerSettings;
@@ -776,7 +802,7 @@ void SMediaViewer::SaveStates(bool bInSaveConfig) const
 					LibraryItem->GetStringValue(),
 					ImageViewers[Index]->GetPanelSettings(),
 					ImageViewers[Index]->GetPaintSettings()
-				});
+					});
 				bAdded = true;
 			}
 		}
@@ -787,35 +813,20 @@ void SMediaViewer::SaveStates(bool bInSaveConfig) const
 		}
 	}
 
-	UMediaViewerLibraryIni& Ini = UMediaViewerLibraryIni::Get();
-	Ini.SetSavedStates({State});
-
-	if (bInSaveConfig)
-	{
-		Ini.SaveConfig();
-	}
+	return State;
 }
 
-void SMediaViewer::LoadState(int32 InIndex)
+void SMediaViewer::LoadState(const FMediaViewerState& InState)
 {
-	TConstArrayView<FMediaViewerState> SavedStates = UMediaViewerLibraryIni::Get().GetSavedStates();
-
-	if (!SavedStates.IsValidIndex(InIndex))
-	{
-		return;
-	}
-
-	const FMediaViewerState& State = SavedStates[InIndex];
-
 	TArray<TSharedPtr<FMediaViewerLibraryItem>> Items;
-	Items.Reserve(State.Images.Num());
+	Items.Reserve(InState.Images.Num());
 	bool bHasValidItem = false;
 
-	for (int32 Index = 0; Index < State.Images.Num(); ++Index)
+	for (int32 Index = 0; Index < InState.Images.Num(); ++Index)
 	{
 		TSharedPtr<FMediaViewerLibraryItem> Item = Library->GetLibrary()->FindItemByValue(
-			State.Images[Index].ImageType, 
-			State.Images[Index].StringValue
+			InState.Images[Index].ImageType,
+			InState.Images[Index].StringValue
 		);
 
 		Items.Add(Item);
@@ -831,28 +842,18 @@ void SMediaViewer::LoadState(int32 InIndex)
 		return;
 	}
 
-	const EAppReturnType::Type Result = FMessageDialog::Open(
-		EAppMsgType::YesNo,
-		LOCTEXT("LoadOldState", "Attempt to open previous images?")
-	);
-
-	if (Result != EAppReturnType::Yes)
-	{
-		return;
-	}
-
-	MediaViewerSettings = State.ViewerSettings;
-	RequestedView = State.ActiveView;
+	MediaViewerSettings = InState.ViewerSettings;
+	RequestedView = InState.ActiveView;
 
-	for (int32 Index = 0; Index < State.Images.Num(); ++Index)
+	for (int32 Index = 0; Index < InState.Images.Num(); ++Index)
 	{
 		if (Items[Index].IsValid())
 		{
 			if (TSharedPtr<FMediaImageViewer> Viewer = Items[Index]->CreateImageViewer())
 			{
 				ImageViewers[Index] = Viewer.ToSharedRef();
-				Viewer->GetPanelSettings() = State.Images[Index].PanelSettings;
-				Viewer->GetPaintSettings() = State.Images[Index].PaintSettings;
+				Viewer->GetPanelSettings() = InState.Images[Index].PanelSettings;
+				Viewer->GetPaintSettings() = InState.Images[Index].PaintSettings;
 			}
 		}
 	}
@@ -860,6 +861,40 @@ void SMediaViewer::LoadState(int32 InIndex)
 	InvalidateView();
 }
 
+void SMediaViewer::SaveBookmark(int32 InIndex, bool bInSaveConfig)
+{
+	UMediaViewerLibraryIni& Ini = UMediaViewerLibraryIni::Get();
+	Ini.SetBookmark(InIndex, GetState());
+
+	if (bInSaveConfig)
+	{
+		Ini.SaveConfig();
+	}
+}
+
+void SMediaViewer::LoadBookmark(int32 InIndex)
+{
+	TConstArrayView<FMediaViewerState> Bookmarks = UMediaViewerLibraryIni::Get().GetBookmarks();
+
+	if (!Bookmarks.IsValidIndex(InIndex))
+	{
+		return;
+	}
+
+	LoadState(Bookmarks[InIndex]);
+}
+
+void SMediaViewer::SaveLastOpenedState(bool bInSaveConfig)
+{
+	UMediaViewerLibraryIni& Ini = UMediaViewerLibraryIni::Get();
+	Ini.SetLastOpenedState(GetState());
+
+	if (bInSaveConfig)
+	{
+		Ini.SaveConfig();
+	}
+}
+
 void SMediaViewer::InvalidateView()
 {
 	bInvalidated = true;
@@ -1034,7 +1069,7 @@ void SMediaViewer::ToggleLockedTransform()
 	MediaViewerSettings.bAreTransformsLocked = !MediaViewerSettings.bAreTransformsLocked;
 }
 
-void SMediaViewer::CheckLoadState()
+void SMediaViewer::CheckLastOpenedState()
 {
 	const FGuid& NullImageId = FNullImageViewer::GetNullImageViewer()->GetInfo().Id;
 
@@ -1046,7 +1081,27 @@ void SMediaViewer::CheckLoadState()
 		}
 	}
 
-	LoadState(0);
+	const EAppReturnType::Type Result = FMessageDialog::Open(
+		EAppMsgType::YesNo,
+		LOCTEXT("LoadOldState", "Attempt to open previous images?")
+	);
+
+	if (Result == EAppReturnType::Yes)
+	{
+		LoadLastOpenedState();
+	}
+}
+
+void SMediaViewer::LoadLastOpenedState()
+{
+	TOptional<FMediaViewerState> LastOpenedState = UMediaViewerLibraryIni::Get().GetLastOpenedState();
+
+	if (!LastOpenedState.IsSet())
+	{
+		return;
+	}
+
+	LoadState(LastOpenedState.GetValue());
 }
 
 bool SMediaViewer::CanTakeSnapshot() const
@@ -1271,6 +1326,45 @@ void SMediaViewer::OpenNewTab()
 	IMediaViewerModule::Get().OpenStandaloneTab();
 }
 
+bool SMediaViewer::CanOpenBookmark(int32 InIndex) const
+{
+	TConstArrayView<FMediaViewerState> Bookmarks = UMediaViewerLibraryIni::Get().GetBookmarks();
+
+	if (!Bookmarks.IsValidIndex(InIndex))
+	{
+		return false;
+	}
+
+	for (const FMediaViewerImageState& Image : Bookmarks[InIndex].Images)
+	{
+		if (!Image.StringValue.IsEmpty())
+		{
+			return true;
+		}
+	}
+
+	return false;
+}
+
+void SMediaViewer::OpenBookmark(int32 InIndex)
+{
+	LoadBookmark(InIndex);
+
+	UE_LOG(LogMediaViewer, Log, TEXT("Bookmark loaded."));
+}
+
+bool SMediaViewer::CanSaveBookmark() const
+{
+	return !IsImageViewerNull(EMediaImageViewerPosition::First) || !IsImageViewerNull(EMediaImageViewerPosition::Second);
+}
+
+void SMediaViewer::SaveBookmark(int32 InIndex)
+{
+	SaveBookmark(InIndex, /* Save config */ true);
+
+	UE_LOG(LogMediaViewer, Log, TEXT("Bookmark saved."));
+}
+
 } // UE::MediaViewer::Private
 
 #undef LOCTEXT_NAMESPACE
diff --git a/Engine/Plugins/Experimental/MediaViewer/Source/MediaViewer/Private/Widgets/SMediaViewer.h b/Engine/Plugins/Experimental/MediaViewer/Source/MediaViewer/Private/Widgets/SMediaViewer.h
--- a/Engine/Plugins/Experimental/MediaViewer/Source/MediaViewer/Private/Widgets/SMediaViewer.h
+++ b/Engine/Plugins/Experimental/MediaViewer/Source/MediaViewer/Private/Widgets/SMediaViewer.h
@@ -17,6 +17,7 @@
 class FUICommandList;
 class SViewportToolBar;
 class UCanvas;
+struct FMediaViewerState;
 
 namespace UE::MediaViewer
 {
@@ -99,9 +100,15 @@ public:
 	/** Returns the bound delegates of this viewer. */
 	TSharedRef<FMediaViewerDelegates> GetDelegates() const;
 
-	void SaveStates(bool bInSaveConfig) const;
+	FMediaViewerState GetState() const;
 
-	void LoadState(int32 InIndex);
+	void LoadState(const FMediaViewerState& InState);
+
+	void SaveBookmark(int32 InIndex, bool bInSaveConfig);
+
+	void LoadBookmark(int32 InIndex);
+
+	void SaveLastOpenedState(bool bInSaveConfig);
 
 	//~ Begin SWidget
 	virtual void Tick(const FGeometry& AllottedGeometry, const double InCurrentTime, const float InDeltaTime) override;
@@ -240,7 +247,7 @@ protected:
 	ECheckBoxState AreLockedTransformEnabled() const;
 	void ToggleLockedTransform();
 
-	void CheckLoadState();
+	void CheckLastOpenedState();
 
 	bool CanTakeSnapshot() const;
 	void TakeSnapshot();
@@ -248,6 +255,13 @@ protected:
 	void DrawImageViewers(UCanvas* InCanvas) const;
 	
 	void OpenNewTab();
+
+	void LoadLastOpenedState();
+
+	bool CanOpenBookmark(int32 InIndex) const;	
+	void OpenBookmark(int32 InIndex);
+	bool CanSaveBookmark() const;
+	void SaveBookmark(int32 InIndex);
 };
 
 } // UE::MediaViewer::Private
diff --git a/Engine/Plugins/Experimental/MediaViewer/Source/MediaViewer/Private/Widgets/SMediaViewerToolbar.cpp b/Engine/Plugins/Experimental/MediaViewer/Source/MediaViewer/Private/Widgets/SMediaViewerToolbar.cpp
--- a/Engine/Plugins/Experimental/MediaViewer/Source/MediaViewer/Private/Widgets/SMediaViewerToolbar.cpp
+++ b/Engine/Plugins/Experimental/MediaViewer/Source/MediaViewer/Private/Widgets/SMediaViewerToolbar.cpp
@@ -387,6 +387,8 @@ TSharedRef<SWidget> SMediaViewerToolbar::MakeScaleMenu(EMediaImageViewerPosition
 		MenuBuilder.AddMenuEntry(Commands.Scale400);
 		MenuBuilder.AddMenuEntry(Commands.Scale800);
 		MenuBuilder.AddSeparator();
+		MenuBuilder.AddMenuEntry(Commands.ScaleUp);
+		MenuBuilder.AddMenuEntry(Commands.ScaleDown);
 		MenuBuilder.AddMenuEntry(Commands.ScaleToFit);
 	}
 
