Media Viewer: bookmark redux

diff --git a/Engine/Plugins/Experimental/MediaViewer/Source/MediaViewer/Private/Widgets/SMediaImageViewerOverlay.cpp b/Engine/Plugins/Experimental/MediaViewer/Source/MediaViewer/Private/Widgets/SMediaImageViewerOverlay.cpp
--- a/Engine/Plugins/Experimental/MediaViewer/Source/MediaViewer/Private/Widgets/SMediaImageViewerOverlay.cpp
+++ b/Engine/Plugins/Experimental/MediaViewer/Source/MediaViewer/Private/Widgets/SMediaImageViewerOverlay.cpp
@@ -1099,7 +1099,7 @@ TSharedRef<SWidget> SMediaImageViewerOverlay::CreateMenu()
 		MenuBuilder.AddSubMenu(
 			LOCTEXT("SaveBookmarkMenu", "Save Bookmark"),
 			LOCTEXT("SaveBookmarkMenuTooltip", "Save a bookmark of the current state."),
-			LoadBookmarksMenu
+			SaveBookmarksMenu
 		);
 	}
 
