Media VIewer: Crash bug fix

diff --git a/Engine/Plugins/Experimental/MediaViewer/Source/MediaViewer/Private/Widgets/SMediaImageViewerOverlay.cpp b/Engine/Plugins/Experimental/MediaViewer/Source/MediaViewer/Private/Widgets/SMediaImageViewerOverlay.cpp
--- a/Engine/Plugins/Experimental/MediaViewer/Source/MediaViewer/Private/Widgets/SMediaImageViewerOverlay.cpp
+++ b/Engine/Plugins/Experimental/MediaViewer/Source/MediaViewer/Private/Widgets/SMediaImageViewerOverlay.cpp
@@ -1025,7 +1025,7 @@ TSharedRef<SWidget> SMediaImageViewerOverlay::CreateMenu()
 
 	FMenuBuilder MenuBuilder(true, CommandList, nullptr, false, &FAppStyle::Get(), false);
 
-	MenuBuilder.BeginSection(TEXT("Image"));
+	MenuBuilder.BeginSection(TEXT("Image"), LOCTEXT("Image", "Image"));
 	{
 		MenuBuilder.AddMenuEntry(Commands.ResetTransform);
 		MenuBuilder.AddMenuEntry(Commands.CopyTransform);
@@ -1033,22 +1033,21 @@ TSharedRef<SWidget> SMediaImageViewerOverlay::CreateMenu()
 	}
 	MenuBuilder.EndSection();
 
-	MenuBuilder.AddMenuSeparator();
-
-	MenuBuilder.BeginSection(TEXT("Media Viewer"));
+	MenuBuilder.BeginSection(TEXT("Media Viewer"), LOCTEXT("MediaViewer", "Media Viewer"));
 	{
 		MenuBuilder.AddMenuEntry(Commands.ToggleOverlay);
 		MenuBuilder.AddMenuEntry(Commands.ToggleLockedTransform);
 		MenuBuilder.AddMenuEntry(Commands.ResetAllTransforms);
 		MenuBuilder.AddMenuEntry(Commands.SwapAB);
 	}
+	MenuBuilder.EndSection();
 
 	FNewMenuDelegate LoadBookmarksMenu = FNewMenuDelegate::CreateLambda(
 		[](FMenuBuilder& InMenuBuilder)
 		{
 			FMediaViewerCommands& Commands = FMediaViewerCommands::Get();
 
-			InMenuBuilder.BeginSection("Bookmarks");
+			InMenuBuilder.BeginSection(TEXT("OpenBookmarks"), LOCTEXT("Bookmarks", "Bookmarks"));
 			{
 				InMenuBuilder.AddMenuEntry(Commands.OpenBookmark1);
 				InMenuBuilder.AddMenuEntry(Commands.OpenBookmark2);
@@ -1070,7 +1069,7 @@ TSharedRef<SWidget> SMediaImageViewerOverlay::CreateMenu()
 		{
 			FMediaViewerCommands& Commands = FMediaViewerCommands::Get();
 
-			InMenuBuilder.BeginSection("Bookmarks");
+			InMenuBuilder.BeginSection(TEXT("SaveBookmarks"), LOCTEXT("Bookmarks", "Bookmarks"));
 			{
 				InMenuBuilder.AddMenuEntry(Commands.SaveBookmark1);
 				InMenuBuilder.AddMenuEntry(Commands.SaveBookmark2);
@@ -1087,8 +1086,7 @@ TSharedRef<SWidget> SMediaImageViewerOverlay::CreateMenu()
 		}
 	);
 
-	MenuBuilder.AddMenuSeparator();
-	MenuBuilder.BeginSection(TEXT("Bookmarks"));
+	MenuBuilder.BeginSection(TEXT("Bookmarks"), LOCTEXT("Bookmarks", "Bookmarks"));
 	{
 		MenuBuilder.AddSubMenu(
 			LOCTEXT("OpenBookmarkMenu", "Open Bookmark"),
@@ -1102,6 +1100,7 @@ TSharedRef<SWidget> SMediaImageViewerOverlay::CreateMenu()
 			SaveBookmarksMenu
 		);
 	}
+	MenuBuilder.EndSection();
 
 	return MenuBuilder.MakeWidget();
 }
