Media Viewer: Moved from Experimental to Media folder.

#virtualized

diff --git a/Engine/Plugins/Media/MediaViewer/Source/MediaViewer/Private/MediaViewerUtils.h b/Engine/Plugins/Media/MediaViewer/Source/MediaViewer/Private/MediaViewerUtils.h
new file mode 100644
--- /dev/null
+++ b/Engine/Plugins/Media/MediaViewer/Source/MediaViewer/Private/MediaViewerUtils.h
@@ -0,0 +1,74 @@
+// Copyright Epic Games, Inc. All Rights Reserved.
+
+#pragma once
+
+#include "Engine/AssetUserData.h"
+
+#include "Containers/ArrayView.h"
+#include "DetailsViewArgs.h"
+#include "HAL/Platform.h"
+#include "IStructureDetailsView.h"
+#include "Math/Color.h"
+#include "Math/MathFwd.h"
+#include "Misc/Optional.h"
+#include "PropertyEditorModule.h"
+#include "Templates/SharedPointerFwd.h"
+#include "Templates/Tuple.h"
+
+#include "MediaViewerUtils.generated.h"
+
+class FNotifyHook;
+class FStructOnScope;
+class FText;
+class IStructureDetailsView;
+class UMaterialInterface;
+class UTextureRenderTarget2D;
+enum EPixelFormat : uint8;
+struct FDetailsViewArgs;
+struct FStructureDetailsViewArgs;
+template <typename T, typename... Ts> class TVariant;
+
+UCLASS()
+class UMediaViewerUserData : public UAssetUserData
+{
+	GENERATED_BODY()
+};
+
+namespace UE::MediaViewer::Private
+{
+
+DECLARE_DELEGATE_OneParam(FRenderComplete, bool /* Was succsesfully rendered */)
+
+class FMediaViewerUtils
+{
+public:
+	static constexpr const TCHAR* TextureRenderMaterial_TextureParameterName = TEXT("Texture");
+
+	/** Returns a UI Material asset that can be used to render a texture to the emissive and opacity channels. */
+	static UMaterialInterface* GetTextureRenderMaterial();
+
+	/** Retrieves the pixel color of a given pixel based on its pixel format. Or none if out of range. */
+	static TOptional<TVariant<FColor, FLinearColor>> GetPixelColor(TArrayView64<uint8> InPixelData, EPixelFormat InPixelFormat, 
+		const FIntPoint& InTextureSize, const FIntPoint& InPixelCoords, int32 InMipLevel);
+
+	/** 
+	 * Creates a render target texture parented to the transient package. 
+	 * The target will have the @see UMediaViewerUserData asset data to differentiate it from other render targets.
+	 */
+	static UTextureRenderTarget2D* CreateRenderTarget(const FIntPoint& InSize, bool bInTransparent);
+
+	/** 
+	 * Creates a render target (@see CreateRenderTarget) and renders the material to it.
+	 * By default it creates a 256x256 render target which is transparent if the material is a UI material.
+	 */
+	static UTextureRenderTarget2D* RenderMaterial(TNotNull<UMaterialInterface*> InMaterial);
+
+	/** Renders the given material to the given render target using FWidgetRenderer and an FSlateMaterialBrush. */
+	static void RenderMaterial(TNotNull<UMaterialInterface*> InMaterial, TNotNull<UTextureRenderTarget2D*> InRenderTarget);
+
+	/** Creates a struct details view based on the given struct. The view will have most settings disabled for a clean view. */
+	static TSharedRef<IStructureDetailsView> CreateStructDetailsView(TSharedRef<FStructOnScope> InStructOnScope, const FText& InCustomName,
+		FNotifyHook* InNotifyHook = nullptr);
+};
+
+} // UE::MediaViewer::Private
