Media Viewer: Added notification when a bookmark has been saved/loaded.
- Added bookmark index to the log when saving/loading bookmark.
- Fixed issue where some UI commands would not work in media context menu.

diff --git a/Engine/Plugins/Media/MediaViewer/Source/MediaViewer/Private/Widgets/SMediaImageViewerOverlay.cpp b/Engine/Plugins/Media/MediaViewer/Source/MediaViewer/Private/Widgets/SMediaImageViewerOverlay.cpp
--- a/Engine/Plugins/Media/MediaViewer/Source/MediaViewer/Private/Widgets/SMediaImageViewerOverlay.cpp
+++ b/Engine/Plugins/Media/MediaViewer/Source/MediaViewer/Private/Widgets/SMediaImageViewerOverlay.cpp
@@ -434,6 +434,14 @@ EVisibility SMediaImageViewerOverlay::GetDragDescriptionVisibility() const
 
 void SMediaImageViewerOverlay::BindCommands()
 {
+	if (Delegates.IsValid() && Delegates->GetCommandList.IsBound())
+	{
+		if (TSharedPtr<FUICommandList> ParentCommandList = Delegates->GetCommandList.Execute())
+		{
+			CommandList->Append(ParentCommandList.ToSharedRef());
+		}
+	}
+
 	const FMediaViewerCommands& Commands = FMediaViewerCommands::Get();
 
 	CommandList->MapAction(Commands.MoveLeft,     FExecuteAction::CreateSP(this, &SMediaImageViewerOverlay::Try_AddOffset, FVector(-10,   0,   0)));
diff --git a/Engine/Plugins/Media/MediaViewer/Source/MediaViewer/Private/Widgets/SMediaViewer.cpp b/Engine/Plugins/Media/MediaViewer/Source/MediaViewer/Private/Widgets/SMediaViewer.cpp
--- a/Engine/Plugins/Media/MediaViewer/Source/MediaViewer/Private/Widgets/SMediaViewer.cpp
+++ b/Engine/Plugins/Media/MediaViewer/Source/MediaViewer/Private/Widgets/SMediaViewer.cpp
@@ -13,6 +13,7 @@
 #include "Engine/TextureRenderTarget2D.h"
 #include "Framework/Application/SlateApplication.h"
 #include "Framework/Notifications/NotificationManager.h"
+#include "Framework/Notifications/NotificationManager.h"
 #include "IContentBrowserSingleton.h"
 #include "ImageViewer/MediaImageViewer.h"
 #include "ImageViewers/NullImageViewer.h"
@@ -33,6 +34,7 @@
 #include "Widgets/Layout/SBox.h"
 #include "Widgets/Layout/SSplitter.h"
 #include "Widgets/Notifications/SNotificationList.h"
+#include "Widgets/Notifications/SNotificationList.h"
 #include "Widgets/SMediaImageViewerOverlay.h"
 #include "Widgets/SMediaViewerLibrary.h"
 #include "Widgets/SMediaViewerLibraryPrivate.h"
@@ -1351,6 +1353,14 @@ void SMediaViewer::OpenBookmark(int32 InIndex)
 	LoadBookmark(InIndex);
 
 	UE_LOG(LogMediaViewer, Log, TEXT("Bookmark %i loaded."), InIndex + 1);
+
+	const FText MessageFormat = LOCTEXT("BookmarkIndexLoaded", "[Media Viewer] Bookmark {0} loaded.");
+
+	FNotificationInfo NotificationInfo(FText::Format(MessageFormat, FText::AsNumber(InIndex + 1)));
+	NotificationInfo.ExpireDuration = 3.f;
+	NotificationInfo.bFireAndForget = true;
+
+	FSlateNotificationManager::Get().AddNotification(NotificationInfo);
 }
 
 bool SMediaViewer::CanSaveBookmark() const
@@ -1363,6 +1373,14 @@ void SMediaViewer::SaveBookmark(int32 InIndex)
 	SaveBookmark(InIndex, /* Save config */ true);
 
 	UE_LOG(LogMediaViewer, Log, TEXT("Bookmark %i saved."), InIndex + 1);
+
+	const FText MessageFormat = LOCTEXT("BookmarkIndexSaved", "[Media Viewer] Bookmark {0} saved.");
+
+	FNotificationInfo NotificationInfo(FText::Format(MessageFormat, FText::AsNumber(InIndex + 1)));
+	NotificationInfo.ExpireDuration = 3.f;
+	NotificationInfo.bFireAndForget = true;
+
+	FSlateNotificationManager::Get().AddNotification(NotificationInfo);
 }
 
 } // UE::MediaViewer::Private
