Media Viewer: Fixed background texture on bookmark load.

diff --git a/Engine/Plugins/Media/MediaViewer/Source/MediaViewer/Private/ImageViewers/MediaImageViewer.cpp b/Engine/Plugins/Media/MediaViewer/Source/MediaViewer/Private/ImageViewers/MediaImageViewer.cpp
--- a/Engine/Plugins/Media/MediaViewer/Source/MediaViewer/Private/ImageViewers/MediaImageViewer.cpp
+++ b/Engine/Plugins/Media/MediaViewer/Source/MediaViewer/Private/ImageViewers/MediaImageViewer.cpp
@@ -54,16 +54,7 @@ void FMediaImageViewer::UpdateId(const FGuid& InId)
 
 void FMediaImageViewer::NotifyPostChange(const FPropertyChangedEvent& InPropertyChangedEvent, FProperty* InPropertyThatChanged)
 {
-	if (UTexture* Texture = PanelSettings.BackgroundTexture.LoadSynchronous())
-	{
-		BackgroundImageBrush.SetResourceObject(Texture);
-		BackgroundImageBrush.ImageSize.X = Texture->GetSurfaceWidth();
-		BackgroundImageBrush.ImageSize.Y = Texture->GetSurfaceHeight();
-	}
-	else
-	{
-		BackgroundImageBrush.SetResourceObject(nullptr);
-	}
+	InitBackgroundTexture();
 }
 
 FVector2D FMediaImageViewer::GetViewerCenter(const FVector2D& InViewerSize) const
@@ -128,6 +119,20 @@ FSlateClippingZone FMediaImageViewer::CreateSlateClippingZone(const FSlateRect&
 	return ClippingZone;
 }
 
+void FMediaImageViewer::InitBackgroundTexture()
+{
+	if (UTexture* Texture = PanelSettings.BackgroundTexture.LoadSynchronous())
+	{
+		BackgroundImageBrush.SetResourceObject(Texture);
+		BackgroundImageBrush.ImageSize.X = Texture->GetSurfaceWidth();
+		BackgroundImageBrush.ImageSize.Y = Texture->GetSurfaceHeight();
+	}
+	else
+	{
+		BackgroundImageBrush.SetResourceObject(nullptr);
+	}
+}
+
 void FMediaImageViewer::Paint(FMediaImageSlatePaintParams& InPaintParams)
 {
 	FVector2D PaintOffset = GetPaintOffset(InPaintParams.ViewerSize, InPaintParams.ViewerPosition);
diff --git a/Engine/Plugins/Media/MediaViewer/Source/MediaViewer/Private/Widgets/SMediaViewer.cpp b/Engine/Plugins/Media/MediaViewer/Source/MediaViewer/Private/Widgets/SMediaViewer.cpp
--- a/Engine/Plugins/Media/MediaViewer/Source/MediaViewer/Private/Widgets/SMediaViewer.cpp
+++ b/Engine/Plugins/Media/MediaViewer/Source/MediaViewer/Private/Widgets/SMediaViewer.cpp
@@ -13,7 +13,6 @@
 #include "Engine/TextureRenderTarget2D.h"
 #include "Framework/Application/SlateApplication.h"
 #include "Framework/Notifications/NotificationManager.h"
-#include "Framework/Notifications/NotificationManager.h"
 #include "IContentBrowserSingleton.h"
 #include "ImageViewer/MediaImageViewer.h"
 #include "ImageViewers/NullImageViewer.h"
@@ -34,7 +33,6 @@
 #include "Widgets/Layout/SBox.h"
 #include "Widgets/Layout/SSplitter.h"
 #include "Widgets/Notifications/SNotificationList.h"
-#include "Widgets/Notifications/SNotificationList.h"
 #include "Widgets/SMediaImageViewerOverlay.h"
 #include "Widgets/SMediaViewerLibrary.h"
 #include "Widgets/SMediaViewerLibraryPrivate.h"
@@ -856,6 +854,7 @@ void SMediaViewer::LoadState(const FMediaViewerState& InState)
 				ImageViewers[Index] = Viewer.ToSharedRef();
 				Viewer->GetPanelSettings() = InState.Images[Index].PanelSettings;
 				Viewer->GetPaintSettings() = InState.Images[Index].PaintSettings;
+				Viewer->InitBackgroundTexture();
 			}
 		}
 	}
diff --git a/Engine/Plugins/Media/MediaViewer/Source/MediaViewer/Public/ImageViewer/MediaImageViewer.h b/Engine/Plugins/Media/MediaViewer/Source/MediaViewer/Public/ImageViewer/MediaImageViewer.h
--- a/Engine/Plugins/Media/MediaViewer/Source/MediaViewer/Public/ImageViewer/MediaImageViewer.h
+++ b/Engine/Plugins/Media/MediaViewer/Source/MediaViewer/Public/ImageViewer/MediaImageViewer.h
@@ -230,6 +230,8 @@ public:
 	const FMediaImagePaintSettings& GetPaintSettings() const { return PaintSettings; }
 	FMediaImagePaintSettings& GetPaintSettings() { return PaintSettings; }
 
+	void InitBackgroundTexture();
+
 	virtual TOptional<TVariant<FColor, FLinearColor>> GetPixelColor(const FIntPoint& InPixelCoords, int32 InMipLevel) const = 0;
 
 	/** Paint the panel and image on to the viewer. */
