Media Viewer: Toggling the overlay now correctly ignores specific widgets, fixing the media source overlay.

diff --git a/Engine/Plugins/Media/MediaViewer/Source/MediaViewer/Private/Widgets/SMediaImageViewerOverlay.cpp b/Engine/Plugins/Media/MediaViewer/Source/MediaViewer/Private/Widgets/SMediaImageViewerOverlay.cpp
--- a/Engine/Plugins/Media/MediaViewer/Source/MediaViewer/Private/Widgets/SMediaImageViewerOverlay.cpp
+++ b/Engine/Plugins/Media/MediaViewer/Source/MediaViewer/Private/Widgets/SMediaImageViewerOverlay.cpp
@@ -30,7 +30,7 @@
 namespace UE::MediaViewer::Private
 {
 
-const FName DropTargetName = TEXT("DropTarget");
+const FName IgnoreOverlayToggle = TEXT("IgnoreOverlayToggle");
 
 SMediaImageViewerOverlay::SMediaImageViewerOverlay()
 	: CommandList(MakeShared<FUICommandList>())
@@ -528,6 +528,7 @@ TSharedRef<SWidget> SMediaImageViewerOverlay::CreateOverlay(bool bInComparisonVi
 			.Text(LOCTEXT("DropTargetMessage", "Drop supported asset or library item here."))
 			.Font(IDetailLayoutBuilder::GetDetailFont())
 			.TextStyle(FAppStyle::Get(), "HintText")
+			.Tag(IgnoreOverlayToggle)
 		];
 
 	if (bInComparisonView)
@@ -537,7 +538,7 @@ TSharedRef<SWidget> SMediaImageViewerOverlay::CreateOverlay(bool bInComparisonVi
 				SNew(SMediaViewerDropTarget, Delegates.ToSharedRef())
 				.Position(Position)
 				.bComparisonView(bInComparisonView)
-				.Tag(DropTargetName)
+				.Tag(IgnoreOverlayToggle)
 			];
 	}
 	else
@@ -545,7 +546,7 @@ TSharedRef<SWidget> SMediaImageViewerOverlay::CreateOverlay(bool bInComparisonVi
 		Overlay->AddSlot()
 			[
 				SNew(SHorizontalBox)
-				.Tag(DropTargetName)
+				.Tag(IgnoreOverlayToggle)
 				+ SHorizontalBox::Slot()
 				.FillWidth(0.25f)
 				.Padding(0.f, 0.f, 0.f, 20.f)
@@ -613,6 +614,8 @@ TSharedRef<SWidget> SMediaImageViewerOverlay::CreateOverlay(bool bInComparisonVi
 			.HAlign(HAlign_Center)
 			.VAlign(VAlign_Top)
 			.Padding(5.f)
+			.Visibility(EVisibility::HitTestInvisible)
+			.Tag(IgnoreOverlayToggle)
 			[
 				SNew(STextBlock)
 				.Font(IDetailLayoutBuilder::GetDetailFont())
@@ -654,7 +657,7 @@ void SMediaImageViewerOverlay::ToggleOverlay()
 		{
 			TSharedRef<SWidget> Child = Children->GetChildAt(Index);
 
-			if (Child->GetTag() != DropTargetName)
+			if (Child->GetTag() != IgnoreOverlayToggle)
 			{
 				Child->SetVisibility(NewVisibility);
 			}
