Media Viewer: When saving a snapshot of a file, it will use a shorter/better file name.

diff --git a/Engine/Plugins/Media/MediaViewer/Source/MediaViewer/Private/Widgets/SMediaViewer.cpp b/Engine/Plugins/Media/MediaViewer/Source/MediaViewer/Private/Widgets/SMediaViewer.cpp
--- a/Engine/Plugins/Media/MediaViewer/Source/MediaViewer/Private/Widgets/SMediaViewer.cpp
+++ b/Engine/Plugins/Media/MediaViewer/Source/MediaViewer/Private/Widgets/SMediaViewer.cpp
@@ -25,6 +25,7 @@
 #include "MediaViewerDelegates.h"
 #include "MediaViewerUtils.h"
 #include "Misc/MessageDialog.h"
+#include "PackageTools.h"
 #include "Sidebar/SSidebar.h"
 #include "Sidebar/SSidebarContainer.h"
 #include "TextureResource.h"
@@ -1158,7 +1159,7 @@ void SMediaViewer::TakeSnapshot()
 	const FString PathStr = CurrentPath.HasInternalPath() ? CurrentPath.GetInternalPathString() : "/Game";
 
 	FSaveAssetDialogConfig SaveAssetDialogConfig;
-	SaveAssetDialogConfig.DialogTitleOverride = LOCTEXT("ExportMaterialTo", "Export Material To");
+	SaveAssetDialogConfig.DialogTitleOverride = LOCTEXT("SaveSnapshotAs", "Save Snapshot As");
 	SaveAssetDialogConfig.DefaultPath = PathStr;
 	SaveAssetDialogConfig.ExistingAssetPolicy = ESaveAssetDialogExistingAssetPolicy::Disallow;
 	SaveAssetDialogConfig.DefaultAssetName = LOCTEXT("MediaViewerSnapshot", "MediaViewerSnapshot").ToString();
@@ -1167,8 +1168,21 @@ void SMediaViewer::TakeSnapshot()
 	{
 		if (TSharedPtr<FMediaImageViewer> ImageViewer = GetImageViewer(static_cast<EMediaImageViewerPosition>(Index)))
 		{
+			if (!ImageViewer->IsValid())
+			{
+				continue;
+			}
+
 			SaveAssetDialogConfig.DefaultAssetName += TEXT("_");
-			SaveAssetDialogConfig.DefaultAssetName += ImageViewer->GetInfo().DisplayName.ToString();
+
+			FString DisplayName = ImageViewer->GetInfo().DisplayName.ToString();
+
+			if (FPaths::FileExists(DisplayName))
+			{
+				DisplayName = FPaths::GetCleanFilename(DisplayName);
+			}
+
+			SaveAssetDialogConfig.DefaultAssetName += UPackageTools::SanitizePackageName(DisplayName);
 		}
 	}
 
