Media Viewer: Standalone tabs now have an icon

diff --git a/Engine/Plugins/Media/MediaViewer/Source/MediaViewer/Private/MediaViewerModule.cpp b/Engine/Plugins/Media/MediaViewer/Source/MediaViewer/Private/MediaViewerModule.cpp
--- a/Engine/Plugins/Media/MediaViewer/Source/MediaViewer/Private/MediaViewerModule.cpp
+++ b/Engine/Plugins/Media/MediaViewer/Source/MediaViewer/Private/MediaViewerModule.cpp
@@ -249,6 +249,8 @@ TSharedPtr<SDockTab> FMediaViewerModule::OpenStandaloneTab()
 			ViewerTab
 		];
 
+	StandaloneTab->SetTabIcon(FAppStyle::Get().GetBrush("Sequencer.Tracks.Media"));
+
 	TabManager->InsertNewDocumentTab(
 		TabId,
 		StandaloneTabId,
