Media Viewer: Scrubbing now calls Scrub instead of Seek to enable smoothing scrubbing.

diff --git a/Engine/Plugins/Media/MediaViewer/Source/MediaViewer/Private/Widgets/SMediaSourceOverlay.cpp b/Engine/Plugins/Media/MediaViewer/Source/MediaViewer/Private/Widgets/SMediaSourceOverlay.cpp
--- a/Engine/Plugins/Media/MediaViewer/Source/MediaViewer/Private/Widgets/SMediaSourceOverlay.cpp
+++ b/Engine/Plugins/Media/MediaViewer/Source/MediaViewer/Private/Widgets/SMediaSourceOverlay.cpp
@@ -1121,7 +1121,7 @@ void SMediaSourceOverlay::ReceiveScrub(const MediaSourceOverlay::Events::FMediaI
 				case EMediaViewerMediaSyncType::RelativeTime:
 					if (SyncedScrubStartTime.IsType<FTimespan>() && InEventParams.ScrubData.IsType<FTimespan>())
 					{
-						MediaPlayer->Seek(FTimespan::FromSeconds(FMath::Clamp(
+						MediaPlayer->Scrub(FTimespan::FromSeconds(FMath::Clamp(
 							(SyncedScrubStartTime.Get<FTimespan>() + InEventParams.ScrubData.Get<FTimespan>()).GetTotalSeconds(),
 							0.0,
 							MediaPlayer->GetDuration().GetTotalSeconds()
@@ -1132,7 +1132,7 @@ void SMediaSourceOverlay::ReceiveScrub(const MediaSourceOverlay::Events::FMediaI
 				case EMediaViewerMediaSyncType::AbsoluteTime:
 					if (InEventParams.ScrubData.IsType<FTimespan>())
 					{
-						MediaPlayer->Seek(FTimespan::FromSeconds(FMath::Clamp(
+						MediaPlayer->Scrub(FTimespan::FromSeconds(FMath::Clamp(
 							InEventParams.ScrubData.Get<FTimespan>().GetTotalSeconds(),
 							0.0,
 							MediaPlayer->GetDuration().GetTotalSeconds()
@@ -1143,7 +1143,7 @@ void SMediaSourceOverlay::ReceiveScrub(const MediaSourceOverlay::Events::FMediaI
 				case EMediaViewerMediaSyncType::RelativeProgress:
 					if (SyncedScrubStartTime.IsType<float>() && InEventParams.ScrubData.IsType<float>())
 					{
-						MediaPlayer->Seek(MediaPlayer->GetDuration() * FMath::Clamp(
+						MediaPlayer->Scrub(MediaPlayer->GetDuration() * FMath::Clamp(
 							SyncedScrubStartTime.Get<float>() + InEventParams.ScrubData.Get<float>(),
 							0.0,
 							1.0
@@ -1154,7 +1154,7 @@ void SMediaSourceOverlay::ReceiveScrub(const MediaSourceOverlay::Events::FMediaI
 				case EMediaViewerMediaSyncType::AbsoluteProgress:
 					if (InEventParams.ScrubData.IsType<float>())
 					{
-						MediaPlayer->Seek(MediaPlayer->GetDuration() * FMath::Clamp(
+						MediaPlayer->Scrub(MediaPlayer->GetDuration() * FMath::Clamp(
 							InEventParams.ScrubData.Get<float>(),
 							0.0,
 							1.0
