Media Viewer: Toggling the overlay while in comparison mode no longer blocks the custom image ui (i.e. video control overlay)

diff --git a/Engine/Plugins/Media/MediaViewer/Source/MediaViewer/Private/Widgets/SMediaImageViewerOverlay.cpp b/Engine/Plugins/Media/MediaViewer/Source/MediaViewer/Private/Widgets/SMediaImageViewerOverlay.cpp
--- a/Engine/Plugins/Media/MediaViewer/Source/MediaViewer/Private/Widgets/SMediaImageViewerOverlay.cpp
+++ b/Engine/Plugins/Media/MediaViewer/Source/MediaViewer/Private/Widgets/SMediaImageViewerOverlay.cpp
@@ -597,6 +597,8 @@ TSharedRef<SWidget> SMediaImageViewerOverlay::CreateOverlay(bool bInComparisonVi
 				.HAlign(HAlign)
 				.VAlign(VAlign_Top)
 				.Padding(5.f)
+				.Visibility(EVisibility::HitTestInvisible)
+				.Tag(IgnoreOverlayToggle)
 				[
 					SNew(STextBlock)
 					.Font(IDetailLayoutBuilder::GetDetailFont())
