Media Viewer: You can now override the mip level of images.

#virtualized

diff --git a/Engine/Plugins/Media/MediaViewer/Source/MediaViewer/Private/ImageViewers/ColorImageViewer.h b/Engine/Plugins/Media/MediaViewer/Source/MediaViewer/Private/ImageViewers/ColorImageViewer.h
--- a/Engine/Plugins/Media/MediaViewer/Source/MediaViewer/Private/ImageViewers/ColorImageViewer.h
+++ b/Engine/Plugins/Media/MediaViewer/Source/MediaViewer/Private/ImageViewers/ColorImageViewer.h
@@ -65,6 +65,7 @@ public:
 
 	//~ Begin FMediaImageViewer
 	virtual TSharedPtr<FMediaViewerLibraryItem> CreateLibraryItem() const override;
+	virtual bool SupportsMip() const override { return false; }
 	virtual TOptional<TVariant<FColor, FLinearColor>> GetPixelColor(const FIntPoint& InPixelCoords, int32 InMipLevel) const override;
 	//~ End FMediaImageViewer
 
diff --git a/Engine/Plugins/Media/MediaViewer/Source/MediaViewer/Private/ImageViewers/MediaImageViewer.cpp b/Engine/Plugins/Media/MediaViewer/Source/MediaViewer/Private/ImageViewers/MediaImageViewer.cpp
--- a/Engine/Plugins/Media/MediaViewer/Source/MediaViewer/Private/ImageViewers/MediaImageViewer.cpp
+++ b/Engine/Plugins/Media/MediaViewer/Source/MediaViewer/Private/ImageViewers/MediaImageViewer.cpp
@@ -6,9 +6,13 @@
 #include "Engine/Texture.h"
 #include "Engine/TextureRenderTarget2D.h"
 #include "GameFramework/Actor.h"
+#include "Materials/Material.h"
+#include "Materials/MaterialInstanceDynamic.h"
 #include "Math/UnrealMathUtility.h"
 #include "Math/Vector2D.h"
+#include "MediaViewerUtils.h"
 #include "Misc/TVariant.h"
+#include "SlateMaterialBrush.h"
 #include "Widgets/SMediaViewer.h"
 
 #include UE_INLINE_GENERATED_CPP_BY_NAME(MediaImageViewer)
@@ -17,6 +21,7 @@ namespace UE::MediaViewer
 {
 
 const FSlateColorBrush FMediaImageViewer::BackgroundColorBrush(FLinearColor::White);
+constexpr const TCHAR* MipRenderMaterial = TEXT("/Script/Engine.Material'/MediaViewer/M_TextureMip.M_TextureMip'");
 
 MEDIAVIEWER_API FText FMediaImageViewer::GetObjectDisplayName(const UObject* InObject)
 {
@@ -77,7 +82,66 @@ FVector2D FMediaImageViewer::GetPaintSize() const
 	return FVector2D(ImageInfo.Size.X, ImageInfo.Size.Y) * GetPaintSettings().Scale;
 }
 
-FSlateClippingZone FMediaImageViewer::CreateSlateClippingZone(const FSlateRect& InCullingRect, float InDPIScale, 
+bool FMediaImageViewer::UpdateMipBrush()
+{
+	if (!Brush.IsValid())
+	{
+		return false;
+	}
+
+	UTexture* BrushTexture = Cast<UTexture>(Brush->GetResourceObject());
+
+	if (!BrushTexture)
+	{
+		return false;
+	}
+
+	if (!MipAdjustedMaterial.IsValid())
+	{
+		UMaterial* MipAdjustedMaterialBase = LoadObject<UMaterial>(GetTransientPackage(), MipRenderMaterial);
+		MipAdjustedMaterial.Reset(UMaterialInstanceDynamic::Create(MipAdjustedMaterialBase, GetTransientPackage()));
+
+		if (!MipAdjustedMaterial.IsValid())
+		{
+			return false;
+		}
+	}
+
+	if (!MipAdjustedRenderTarget.IsValid())
+	{
+		MipAdjustedRenderTarget.Reset(Private::FMediaViewerUtils::CreateRenderTarget(ImageInfo.Size, /* Transparent */ true));
+
+		if (!MipAdjustedRenderTarget.IsValid())
+		{
+			return false;
+		}
+	}
+	// If it was just created, we don't need to change the size.
+	else if (MipAdjustedRenderTarget->SizeX != ImageInfo.Size.X || MipAdjustedRenderTarget->SizeY != ImageInfo.Size.Y)
+	{
+		MipAdjustedRenderTarget->ResizeTarget(ImageInfo.Size.X, ImageInfo.Size.Y);
+	}
+
+	if (!MipAdjustedBrush.IsValid())
+	{
+		const FVector2D MaterialBrushSize = FVector2D(ImageInfo.Size.X, ImageInfo.Size.Y);
+		MipAdjustedBrush = MakeShared<FSlateMaterialBrush>(*MipAdjustedMaterial.Get(), MaterialBrushSize);
+	}
+	// If it was just created, we don't need to change the size.
+	else if (!FMath::IsNearlyEqual(MipAdjustedBrush->ImageSize.X, ImageInfo.Size.X)
+		|| !FMath::IsNearlyEqual(MipAdjustedBrush->ImageSize.Y, ImageInfo.Size.Y))
+	{
+		const FVector2D MaterialBrushSize = FVector2D(ImageInfo.Size.X, ImageInfo.Size.Y);
+		MipAdjustedBrush->SetImageSize(MaterialBrushSize);
+	}
+
+	MipAdjustedMaterial->SetTextureParameterValue(TEXT("Texture"), BrushTexture);
+	MipAdjustedMaterial->SetScalarParameterValue(TEXT("MipLevel"), static_cast<float>(PaintSettings.GetMipLevel()));
+
+	return true;
+}
+
+FSlateClippingZone FMediaImageViewer::CreateSlateClippingZone(const FSlateRect& InCullingRect, float InDPIScale,
 	const FVector2D& InViewerPosition, EOrientation InOrientation, const FFloatRange& InUVRange) const
 {
 	constexpr float WindowBorderPadding = 4.f;
@@ -196,7 +260,17 @@ void FMediaImageViewer::PaintPanel(FMediaImageSlatePaintParams& InPaintParams, c
 
 void FMediaImageViewer::PaintImage(FMediaImageSlatePaintParams& InPaintParams, const FMediaImageSlatePaintGeometry& InPaintGeometry)
 {
-	if (!Brush.IsValid())
+	TSharedPtr<FSlateBrush> BrushToUse = Brush;
+
+	if (SupportsMip() && PaintSettings.bEnableMipOverride)
+	{
+		if (UpdateMipBrush())
+		{
+			BrushToUse = MipAdjustedBrush;
+		}
+	}
+
+	if (!BrushToUse.IsValid())
 	{
 		return;
 	}
@@ -205,7 +279,7 @@ void FMediaImageViewer::PaintImage(FMediaImageSlatePaintParams& InPaintParams, c
 		InPaintParams.DrawElements,
 		InPaintParams.LayerId,
 		InPaintGeometry.Geometry,
-		Brush.Get(),
+		BrushToUse.Get(),
 		DrawEffects,
 		GetPaintSettings().Rotation.Yaw * 2.f * PI / 360.f,
 		{},
@@ -276,14 +350,17 @@ void FMediaImageViewer::PaintPanel(const FMediaImageCanvasPaintParams& InPaintPa
 
 void FMediaImageViewer::PaintImage(const FMediaImageCanvasPaintParams& InPaintParams, const FMediaImageCanvasPaintGeometry& InPaintGeometry)
 {
-	if (!Brush.IsValid())
+	TSharedPtr<FSlateBrush> BrushToUse = Brush;
+
+	if (SupportsMip() && PaintSettings.bEnableMipOverride)
 	{
-		return;
+		if (UpdateMipBrush())
+		{
+			BrushToUse = MipAdjustedBrush;
+		}
 	}
 
-	UTexture* ForegroundTexture = Cast<UTexture>(Brush->GetResourceObject());
-
-	if (!ForegroundTexture)
+	if (!BrushToUse.IsValid())
 	{
 		return;
 	}
@@ -317,16 +394,31 @@ void FMediaImageViewer::PaintImage(const FMediaImageCanvasPaintParams& InPaintPa
 			const float PositionX = InPaintGeometry.Position.X + (ImageUVRangeMinX * InPaintGeometry.Size.X);
 			const float SizeX = InPaintGeometry.Size.X * ImageUVRangeSize;
 
-			InPaintParams.Canvas->K2_DrawTexture(
-				ForegroundTexture,
-				{PositionX, InPaintGeometry.Position.Y},
-				{SizeX, InPaintGeometry.Size.Y},
-				FVector2D(ImageUVRangeMinX, 0),
-				FVector2D(ImageUVRangeSize, 1),
-				GetPaintSettings().Tint * FLinearColor(1.f, 1.f, 1.f, InPaintParams.ImageOpacity),
-				BLEND_Translucent,
-				GetPaintSettings().Rotation.Yaw
-			);
+			if (UTexture* ForegroundTexture = Cast<UTexture>(BrushToUse->GetResourceObject()))
+			{
+				InPaintParams.Canvas->K2_DrawTexture(
+					ForegroundTexture,
+					{ PositionX, InPaintGeometry.Position.Y },
+					{ SizeX, InPaintGeometry.Size.Y },
+					FVector2D(ImageUVRangeMinX, 0),
+					FVector2D(ImageUVRangeSize, 1),
+					GetPaintSettings().Tint * FLinearColor(1.f, 1.f, 1.f, InPaintParams.ImageOpacity),
+					BLEND_Translucent,
+					GetPaintSettings().Rotation.Yaw
+				);
+			}
+			else if (UMaterialInterface* ForegroundMaterial = Cast<UMaterialInterface>(BrushToUse->GetResourceObject()))
+			{
+				InPaintParams.Canvas->K2_DrawMaterial(
+					ForegroundMaterial,
+					{ PositionX, InPaintGeometry.Position.Y },
+					{ SizeX, InPaintGeometry.Size.Y },
+					FVector2D(ImageUVRangeMinX, 0),
+					FVector2D(ImageUVRangeSize, 1),
+					GetPaintSettings().Rotation.Yaw
+				);
+			}
+
 			break;
 		}
 
@@ -355,16 +447,30 @@ void FMediaImageViewer::PaintImage(const FMediaImageCanvasPaintParams& InPaintPa
 			const float PositionY = InPaintGeometry.Position.Y + (ImageUVRangeMinY * InPaintGeometry.Size.Y);
 			const float SizeY = InPaintGeometry.Size.Y * ImageUVRangeSize;
 
-			InPaintParams.Canvas->K2_DrawTexture(
-				ForegroundTexture,
-				{InPaintGeometry.Position.X, PositionY},
-				{InPaintGeometry.Size.X, SizeY},
-				FVector2D(0, ImageUVRangeMinY),
-				FVector2D(1, ImageUVRangeSize),
-				GetPaintSettings().Tint * FLinearColor(1.f, 1.f, 1.f, InPaintParams.ImageOpacity),
-				BLEND_Translucent,
-				GetPaintSettings().Rotation.Yaw
-			);
+			if (UTexture* ForegroundTexture = Cast<UTexture>(BrushToUse->GetResourceObject()))
+			{
+				InPaintParams.Canvas->K2_DrawTexture(
+					ForegroundTexture,
+					{ InPaintGeometry.Position.X, PositionY },
+					{ InPaintGeometry.Size.X, SizeY },
+					FVector2D(0, ImageUVRangeMinY),
+					FVector2D(1, ImageUVRangeSize),
+					GetPaintSettings().Tint * FLinearColor(1.f, 1.f, 1.f, InPaintParams.ImageOpacity),
+					BLEND_Translucent,
+					GetPaintSettings().Rotation.Yaw
+				);
+			}
+			else if (UMaterialInterface* ForegroundMaterial = Cast<UMaterialInterface>(BrushToUse->GetResourceObject()))
+			{
+				InPaintParams.Canvas->K2_DrawMaterial(
+					ForegroundMaterial,
+					{ InPaintGeometry.Position.X, PositionY },
+					{ InPaintGeometry.Size.X, SizeY },
+					FVector2D(0, ImageUVRangeMinY),
+					FVector2D(1, ImageUVRangeSize),
+					GetPaintSettings().Rotation.Yaw
+				);
+			}
 
 			break;
 		}
diff --git a/Engine/Plugins/Media/MediaViewer/Source/MediaViewer/Private/ImageViewers/MediaSourceImageViewer.cpp b/Engine/Plugins/Media/MediaViewer/Source/MediaViewer/Private/ImageViewers/MediaSourceImageViewer.cpp
--- a/Engine/Plugins/Media/MediaViewer/Source/MediaViewer/Private/ImageViewers/MediaSourceImageViewer.cpp
+++ b/Engine/Plugins/Media/MediaViewer/Source/MediaViewer/Private/ImageViewers/MediaSourceImageViewer.cpp
@@ -306,8 +306,15 @@ FMediaSourceImageViewer::FMediaSourceImageViewer(const FGuid& InId, TNotNull<UMe
 
 		if (TScriptInterface<IMediaStreamPlayer> Player = MediaStream->GetPlayer())
 		{
+			FMediaStreamTextureConfig TextureConfig = Player->GetTextureConfig();
+			TextureConfig.bEnableMipGen = true;
+			
+			Player->SetTextureConfig(TextureConfig);
+
 			if (UMediaTexture* MediaTexture = Player->GetMediaTexture())
 			{
+				MediaTexture->MipGenSettings = TextureMipGenSettings::TMGS_Unfiltered;
+
 				ImageInfo.Size.X = MediaTexture->GetSurfaceWidth();
 				ImageInfo.Size.Y = MediaTexture->GetSurfaceHeight();
 
@@ -516,6 +523,8 @@ void FMediaSourceImageViewer::PaintImage(FMediaImageSlatePaintParams& InPaintPar
 		{
 			ImageInfo.Size.X = MediaTexture->GetSurfaceWidth();
 			ImageInfo.Size.Y = MediaTexture->GetSurfaceHeight();
+
+			MediaTexture->MipGenSettings = MediaSourceSettings.MipGenType;
 		}
 	}
 
diff --git a/Engine/Plugins/Media/MediaViewer/Source/MediaViewer/Private/ImageViewers/MediaSourceImageViewer.h b/Engine/Plugins/Media/MediaViewer/Source/MediaViewer/Private/ImageViewers/MediaSourceImageViewer.h
--- a/Engine/Plugins/Media/MediaViewer/Source/MediaViewer/Private/ImageViewers/MediaSourceImageViewer.h
+++ b/Engine/Plugins/Media/MediaViewer/Source/MediaViewer/Private/ImageViewers/MediaSourceImageViewer.h
@@ -7,6 +7,7 @@
 #include "Library/MediaViewerLibraryItem.h"
 #include "UObject/GCObject.h"
 
+#include "Engine/TextureDefines.h"
 #include "ImageViewers/TextureSampleCache.h"
 #include "Math/Color.h"
 #include "Misc/Optional.h"
@@ -26,6 +27,9 @@ struct FMediaSourceImageViewerSettings
 
 	UPROPERTY(VisibleAnywhere, Category = "Texture")
 	TObjectPtr<UMediaSource> MediaSource = nullptr;
+
+	UPROPERTY(VisibleAnywhere, Category = "Texture")
+	TEnumAsByte<TextureMipGenSettings> MipGenType = TextureMipGenSettings::TMGS_Unfiltered;
 };
 
 namespace UE::MediaViewer::Private
diff --git a/Engine/Plugins/Media/MediaViewer/Source/MediaViewer/Private/MediaViewerCommands.cpp b/Engine/Plugins/Media/MediaViewer/Source/MediaViewer/Private/MediaViewerCommands.cpp
--- a/Engine/Plugins/Media/MediaViewer/Source/MediaViewer/Private/MediaViewerCommands.cpp
+++ b/Engine/Plugins/Media/MediaViewer/Source/MediaViewer/Private/MediaViewerCommands.cpp
@@ -144,6 +144,30 @@ void FMediaViewerCommands::RegisterCommands()
 		FInputChord(EKeys::Hyphen)
 	);
 
+	UI_COMMAND(
+		ToggleMipOverride,
+		"Toggle Mip Override",
+		"Toggles whether to override the mip level of the image.",
+		EUserInterfaceActionType::ToggleButton,
+		FInputChord(EKeys::M)
+	);
+
+	UI_COMMAND(
+		IncreaseMipLevel,
+		"Increase Mip Level",
+		"Increase the Mip level by 1 (decreasing quality).",
+		EUserInterfaceActionType::Button,
+		FInputChord(EKeys::Period)
+	);
+
+	UI_COMMAND(
+		DecreaseMipLevel,
+		"Decrease Mip Level",
+		"Decrease the Mip level by 1 (increasing quality).",
+		EUserInterfaceActionType::Button,
+		FInputChord(EKeys::Comma)
+	);
+
 	UI_COMMAND(
 		ScaleToFit, 
 		"Scale to Fit", 
@@ -355,7 +379,7 @@ void FMediaViewerCommands::RegisterCommands()
 		"Open a new standalone tab.",
 		EUserInterfaceActionType::Button,
 		FInputChord(EKeys::N, EModifierKey::Control),
-		FInputChord(EKeys::M, EModifierKey::Command)
+		FInputChord(EKeys::N, EModifierKey::Command)
 	);
 
 	UI_COMMAND(
diff --git a/Engine/Plugins/Media/MediaViewer/Source/MediaViewer/Private/MediaViewerCommands.h b/Engine/Plugins/Media/MediaViewer/Source/MediaViewer/Private/MediaViewerCommands.h
--- a/Engine/Plugins/Media/MediaViewer/Source/MediaViewer/Private/MediaViewerCommands.h
+++ b/Engine/Plugins/Media/MediaViewer/Source/MediaViewer/Private/MediaViewerCommands.h
@@ -35,6 +35,10 @@ public:
 	TSharedPtr<FUICommandInfo> ScaleUp;
 	TSharedPtr<FUICommandInfo> ScaleDown;
 
+	TSharedPtr<FUICommandInfo> ToggleMipOverride;
+	TSharedPtr<FUICommandInfo> IncreaseMipLevel;
+	TSharedPtr<FUICommandInfo> DecreaseMipLevel;
+
 	TSharedPtr<FUICommandInfo> Scale12;
 	TSharedPtr<FUICommandInfo> Scale25;
 	TSharedPtr<FUICommandInfo> Scale50;
diff --git a/Engine/Plugins/Media/MediaViewer/Source/MediaViewer/Private/Widgets/SMediaImageViewerOverlay.cpp b/Engine/Plugins/Media/MediaViewer/Source/MediaViewer/Private/Widgets/SMediaImageViewerOverlay.cpp
--- a/Engine/Plugins/Media/MediaViewer/Source/MediaViewer/Private/Widgets/SMediaImageViewerOverlay.cpp
+++ b/Engine/Plugins/Media/MediaViewer/Source/MediaViewer/Private/Widgets/SMediaImageViewerOverlay.cpp
@@ -461,6 +461,16 @@ void SMediaImageViewerOverlay::BindCommands()
 	CommandList->MapAction(Commands.ScaleUp, FExecuteAction::CreateSP(this, &SMediaImageViewerOverlay::ScaleUp));
 	CommandList->MapAction(Commands.ScaleDown, FExecuteAction::CreateSP(this, &SMediaImageViewerOverlay::ScaleDown));
 
+	CommandList->MapAction(
+		Commands.ToggleMipOverride,
+		FExecuteAction::CreateSP(this, &SMediaImageViewerOverlay::ToggleMipOverride),
+		FCanExecuteAction(),
+		FGetActionCheckState::CreateSP(this, &SMediaImageViewerOverlay::GetMipOverrideState)
+	);
+
+	CommandList->MapAction(Commands.IncreaseMipLevel, FExecuteAction::CreateSP(this, &SMediaImageViewerOverlay::IncreaseMipLevel));
+	CommandList->MapAction(Commands.DecreaseMipLevel, FExecuteAction::CreateSP(this, &SMediaImageViewerOverlay::DecreaseMipLevel));
+
 	CommandList->MapAction(
 		Commands.ScaleToFit, 
 		FExecuteAction::CreateSP(this, &SMediaImageViewerOverlay::ScaleToFit, /* Use transform lock */ true)
@@ -896,6 +906,48 @@ void SMediaImageViewerOverlay::ScaleToFit(bool bInUseTransformLock)
 	}
 }
 
+ECheckBoxState SMediaImageViewerOverlay::GetMipOverrideState() const
+{
+	if (TSharedPtr<FMediaImageViewer> ImageViewer = Delegates->GetImageViewer.Execute(Position))
+	{
+		return ImageViewer->GetPaintSettings().bEnableMipOverride
+			? ECheckBoxState::Checked
+			: ECheckBoxState::Unchecked;
+	}
+
+	return ECheckBoxState::Undetermined;
+}
+
+void SMediaImageViewerOverlay::ToggleMipOverride()
+{
+	if (TSharedPtr<FMediaImageViewer> ImageViewer = Delegates->GetImageViewer.Execute(Position))
+	{
+		ImageViewer->GetPaintSettings().bEnableMipOverride = !ImageViewer->GetPaintSettings().bEnableMipOverride;
+	}
+}
+
+void SMediaImageViewerOverlay::IncreaseMipLevel()
+{
+	if (TSharedPtr<FMediaImageViewer> ImageViewer = Delegates->GetImageViewer.Execute(Position))
+	{
+		if (ImageViewer->GetPaintSettings().MipOverride < 10)
+		{
+			ImageViewer->GetPaintSettings().MipOverride++;
+		}
+	}
+}
+
+void SMediaImageViewerOverlay::DecreaseMipLevel()
+{
+	if (TSharedPtr<FMediaImageViewer> ImageViewer = Delegates->GetImageViewer.Execute(Position))
+	{
+		if (ImageViewer->GetPaintSettings().MipOverride > 0)
+		{
+			ImageViewer->GetPaintSettings().MipOverride--;
+		}
+	}
+}
+
 void SMediaImageViewerOverlay::Try_ResetTransform()
 {
 	if (Delegates->AreTransformsLocked.Execute())
@@ -936,7 +988,7 @@ uint8 SMediaImageViewerOverlay::GetMipLevel() const
 {
 	if (TSharedPtr<FMediaImageViewer> ImageViewer = Delegates->GetImageViewer.Execute(Position))
 	{
-		return ImageViewer->GetPaintSettings().MipLevel;
+		return ImageViewer->GetPaintSettings().GetMipLevel();
 	}
 
 	return 0;
@@ -962,7 +1014,7 @@ void SMediaImageViewerOverlay::SetMipLevel(uint8 InMipLevel)
 	{
 		if (InMipLevel < ImageViewer->GetInfo().MipCount)
 		{
-			ImageViewer->GetPaintSettings().MipLevel = InMipLevel;
+			ImageViewer->GetPaintSettings().MipOverride = InMipLevel;
 		}
 	}
 }
@@ -990,7 +1042,7 @@ void SMediaImageViewerOverlay::CopyColor()
 		return;
 	}
 
-	const TOptional<TVariant<FColor, FLinearColor>> PixelColor = ImageViewer->GetPixelColor(PixelCoordinates, ImageViewer->GetPaintSettings().MipLevel);
+	const TOptional<TVariant<FColor, FLinearColor>> PixelColor = ImageViewer->GetPixelColor(PixelCoordinates, ImageViewer->GetPaintSettings().GetMipLevel());
 
 	if (!PixelColor.IsSet())
 	{
@@ -1043,6 +1095,10 @@ TSharedRef<SWidget> SMediaImageViewerOverlay::CreateMenu()
 		MenuBuilder.AddMenuEntry(Commands.ResetTransform);
 		MenuBuilder.AddMenuEntry(Commands.CopyTransform);
 		MenuBuilder.AddMenuEntry(Commands.AddToLibrary);
+		MenuBuilder.AddSeparator();
+		MenuBuilder.AddMenuEntry(Commands.ToggleMipOverride);
+		MenuBuilder.AddMenuEntry(Commands.IncreaseMipLevel);
+		MenuBuilder.AddMenuEntry(Commands.DecreaseMipLevel);
 	}
 	MenuBuilder.EndSection();
 
diff --git a/Engine/Plugins/Media/MediaViewer/Source/MediaViewer/Private/Widgets/SMediaImageViewerOverlay.h b/Engine/Plugins/Media/MediaViewer/Source/MediaViewer/Private/Widgets/SMediaImageViewerOverlay.h
--- a/Engine/Plugins/Media/MediaViewer/Source/MediaViewer/Private/Widgets/SMediaImageViewerOverlay.h
+++ b/Engine/Plugins/Media/MediaViewer/Source/MediaViewer/Private/Widgets/SMediaImageViewerOverlay.h
@@ -7,6 +7,7 @@
 class FUICommandList;
 class SOverlay;
 class SVerticalBox;
+enum class ECheckBoxState : uint8;
 struct FMediaViewerLibraryItem;
 
 namespace UE::MediaViewer
@@ -147,6 +148,11 @@ protected:
 	void Try_MultipyScale(float InMultiple);
 	void MultiplyScale(float InMultiple);
 
+	ECheckBoxState GetMipOverrideState() const;
+	void ToggleMipOverride();
+	void IncreaseMipLevel();
+	void DecreaseMipLevel();
+
 	/** Interact with the image viewer's entire transform. */
 	void ResetTransform();
 	bool CanResetAllTransforms();
diff --git a/Engine/Plugins/Media/MediaViewer/Source/MediaViewer/Private/Widgets/SMediaImageViewerStatusBar.cpp b/Engine/Plugins/Media/MediaViewer/Source/MediaViewer/Private/Widgets/SMediaImageViewerStatusBar.cpp
--- a/Engine/Plugins/Media/MediaViewer/Source/MediaViewer/Private/Widgets/SMediaImageViewerStatusBar.cpp
+++ b/Engine/Plugins/Media/MediaViewer/Source/MediaViewer/Private/Widgets/SMediaImageViewerStatusBar.cpp
@@ -150,7 +150,7 @@ FText SMediaImageViewerStatusBar::GetColorPickerLabel() const
 		return INVTEXT("-");
 	}
 
-	const TOptional<TVariant<FColor, FLinearColor>> PixelColor = ImageViewer->GetPixelColor(PixelCoordinates, ImageViewer->GetPaintSettings().MipLevel);
+	const TOptional<TVariant<FColor, FLinearColor>> PixelColor = ImageViewer->GetPixelColor(PixelCoordinates, ImageViewer->GetPaintSettings().GetMipLevel());
 
 	if (!PixelColor.IsSet())
 	{
diff --git a/Engine/Plugins/Media/MediaViewer/Source/MediaViewer/Public/ImageViewer/MediaImageViewer.h b/Engine/Plugins/Media/MediaViewer/Source/MediaViewer/Public/ImageViewer/MediaImageViewer.h
--- a/Engine/Plugins/Media/MediaViewer/Source/MediaViewer/Public/ImageViewer/MediaImageViewer.h
+++ b/Engine/Plugins/Media/MediaViewer/Source/MediaViewer/Public/ImageViewer/MediaImageViewer.h
@@ -31,7 +31,9 @@ class FUICommandList;
 class FWidgetStyle;
 class SWidget;
 class UCanvas;
+class UMaterialInstanceDynamic;
 class UTexture;
+class UTextureRenderTarget2D;
 struct FFloatRange;
 struct FGeometry;
 struct FInputEventState;
@@ -90,18 +92,38 @@ struct FMediaImagePaintSettings
 	UPROPERTY(EditAnywhere, Config, Category = "Media")
 	float Scale = 1.f;
 
+	/** Whether to enable the MipOverride */
+	UPROPERTY(EditAnywhere, Config, Category = "Media", meta = (nlineEditConditionToggle))
+	bool bEnableMipOverride = false;
+
 	/**
 	 * The requested mip level. If the mip isn't available, the highest available value will be used.
-	 * Not currently used.
 	 */
-	UPROPERTY(Config)
-	uint8 MipLevel = 0;
+	UPROPERTY(EditAnywhere, Config, Category = "Media", meta = (EditCondition = "bEnableMipOverride", UIMin = 0, ClampMin = 0, UIMax = 10, ClampMax = 10))
+	uint8 MipOverride = 0;
 
 	/**
 	 * Color to tint the image brush with.
 	 */
 	UPROPERTY(EditAnywhere, Config, Category = "Media")
 	FLinearColor Tint = FLinearColor::White;
+
+	/** Returns a Mip value, if enabled. */
+	TOptional<uint8> GetMipOverride() const
+	{
+		if (bEnableMipOverride)
+		{
+			return MipOverride;
+		}
+
+		return {};
+	}
+
+	/** Always returns a Mip value, even if unset, will default to 0. */
+	uint8 GetMipLevel() const
+	{
+		return GetMipOverride().Get(0);
+	}
 };
 
 namespace UE::MediaViewer
@@ -214,6 +236,8 @@ public:
 	const FMediaImageViewerInfo& GetInfo() const { return ImageInfo; }
 
 	virtual bool IsValid() const { return ImageInfo.Id.IsValid(); }
+
+	virtual bool SupportsMip() const { return true; }
 	
 	/** 
 	 * Called, for instance, when an image is dropped into the viewer and it creates an image viewer,
@@ -324,6 +348,11 @@ protected:
 	FSlateImageBrush BackgroundImageBrush;
 	TSharedPtr<FSlateBrush> Brush;
 	ESlateDrawEffect DrawEffects = ESlateDrawEffect::NoPixelSnapping;
+	TStrongObjectPtr<UMaterialInstanceDynamic> MipAdjustedMaterial;
+	TStrongObjectPtr<UTextureRenderTarget2D> MipAdjustedRenderTarget;
+	TSharedPtr<FSlateBrush> MipAdjustedBrush;
+
+	bool UpdateMipBrush();
 
 	FSlateClippingZone CreateSlateClippingZone(const FSlateRect& InCullingRect, float InDPIScale, const FVector2D& InViewerPosition, 
 		EOrientation InOrientation, const FFloatRange& InUVRange) const;
