Media Viewer: Fixed typo in uproperty.

diff --git a/Engine/Plugins/Media/MediaViewer/Source/MediaViewer/Public/ImageViewer/MediaImageViewer.h b/Engine/Plugins/Media/MediaViewer/Source/MediaViewer/Public/ImageViewer/MediaImageViewer.h
--- a/Engine/Plugins/Media/MediaViewer/Source/MediaViewer/Public/ImageViewer/MediaImageViewer.h
+++ b/Engine/Plugins/Media/MediaViewer/Source/MediaViewer/Public/ImageViewer/MediaImageViewer.h
@@ -93,7 +93,7 @@ struct FMediaImagePaintSettings
 	float Scale = 1.f;
 
 	/** Whether to enable the MipOverride */
-	UPROPERTY(EditAnywhere, Config, Category = "Media", meta = (nlineEditConditionToggle))
+	UPROPERTY(EditAnywhere, Config, Category = "Media", meta = (InlineEditConditionToggle))
 	bool bEnableMipOverride = false;
 
 	/**
