Media Viewer: Reworked object referencing to account for level change GC and asset deletion.

diff --git a/Engine/Plugins/Media/MediaViewer/Source/MediaViewer/Private/ImageViewers/ColorImageViewer.cpp b/Engine/Plugins/Media/MediaViewer/Source/MediaViewer/Private/ImageViewers/ColorImageViewer.cpp
--- a/Engine/Plugins/Media/MediaViewer/Source/MediaViewer/Private/ImageViewers/ColorImageViewer.cpp
+++ b/Engine/Plugins/Media/MediaViewer/Source/MediaViewer/Private/ImageViewers/ColorImageViewer.cpp
@@ -148,6 +148,12 @@ TOptional<TVariant<FColor, FLinearColor>> FColorImageViewer::GetPixelColor(const
 	return {};
 }
 
+FString FColorImageViewer::GetReferencerName() const
+{
+	static const FString ReferencerName = TEXT("FColorImageViewer");
+	return ReferencerName;
+}
+
 } // UE::MediaViewer::Private
 
 #undef LOCTEXT_NAMESPACE
diff --git a/Engine/Plugins/Media/MediaViewer/Source/MediaViewer/Private/ImageViewers/ColorImageViewer.h b/Engine/Plugins/Media/MediaViewer/Source/MediaViewer/Private/ImageViewers/ColorImageViewer.h
--- a/Engine/Plugins/Media/MediaViewer/Source/MediaViewer/Private/ImageViewers/ColorImageViewer.h
+++ b/Engine/Plugins/Media/MediaViewer/Source/MediaViewer/Private/ImageViewers/ColorImageViewer.h
@@ -69,6 +69,10 @@ public:
 	virtual TOptional<TVariant<FColor, FLinearColor>> GetPixelColor(const FIntPoint& InPixelCoords, int32 InMipLevel) const override;
 	//~ End FMediaImageViewer
 
+	//~ Begin FGCObject
+	virtual FString GetReferencerName() const override;
+	//~ End FGCObject
+
 protected:
 	static TSharedRef<FSlateColorBrush> ColorBrush;
 };
diff --git a/Engine/Plugins/Media/MediaViewer/Source/MediaViewer/Private/ImageViewers/LevelEditorViewportImageViewer.h b/Engine/Plugins/Media/MediaViewer/Source/MediaViewer/Private/ImageViewers/LevelEditorViewportImageViewer.h
--- a/Engine/Plugins/Media/MediaViewer/Source/MediaViewer/Private/ImageViewers/LevelEditorViewportImageViewer.h
+++ b/Engine/Plugins/Media/MediaViewer/Source/MediaViewer/Private/ImageViewers/LevelEditorViewportImageViewer.h
@@ -5,7 +5,6 @@
 #include "ImageViewer/IMediaImageViewerFactory.h"
 #include "ImageViewers/SceneViewportImageViewer.h"
 #include "Library/MediaViewerLibraryItem.h"
-#include "UObject/GCObject.h"
 
 namespace UE::MediaViewer::Private
 {
diff --git a/Engine/Plugins/Media/MediaViewer/Source/MediaViewer/Private/ImageViewers/MaterialInterfaceImageViewer.cpp b/Engine/Plugins/Media/MediaViewer/Source/MediaViewer/Private/ImageViewers/MaterialInterfaceImageViewer.cpp
--- a/Engine/Plugins/Media/MediaViewer/Source/MediaViewer/Private/ImageViewers/MaterialInterfaceImageViewer.cpp
+++ b/Engine/Plugins/Media/MediaViewer/Source/MediaViewer/Private/ImageViewers/MaterialInterfaceImageViewer.cpp
@@ -42,12 +42,12 @@ bool FMaterialInterfaceImageViewer::FFactory::SupportsAsset(const FAssetData& In
 
 		while (CurrentData.IsValid() && CurrentData.AssetClassPath != UMaterial::StaticClass()->GetClassPathName())
 		{
-			FName Parent = InAssetData.GetTagValueRef<FName>(GET_MEMBER_NAME_CHECKED(UMaterialInstance, Parent));
+			FName Parent = CurrentData.GetTagValueRef<FName>(GET_MEMBER_NAME_CHECKED(UMaterialInstance, Parent));
 			CurrentData = AssetRegistry.Get().GetAssetByObjectPath(FSoftObjectPath(Parent.ToString()));
 
 			++ParentCount;
 
-			if (ParentCount == MaxParentCheck)
+			if (ParentCount >= MaxParentCheck)
 			{
 				UE_LOG(LogMediaViewer, Error, TEXT("Unable to resolve material parent for \"{0}\"."), *InAssetData.GetObjectPathString());
 				break;
@@ -359,9 +359,11 @@ FString FMaterialInterfaceImageViewer::GetReferencerName() const
 	return ReferencerName;
 }
 
-void FMaterialInterfaceImageViewer::AddReferencedObjects(FReferenceCollector& Collector)
+void FMaterialInterfaceImageViewer::AddReferencedObjects(FReferenceCollector& InCollector)
 {
-	Collector.AddPropertyReferencesWithStructARO(
+	FMediaImageViewer::AddReferencedObjects(InCollector);
+
+	InCollector.AddPropertyReferencesWithStructARO(
 		FMaterialInterfaceImageViewerSettings::StaticStruct(),
 		&MaterialSettings
 	);
@@ -438,6 +440,24 @@ void FMaterialInterfaceImageViewer::OnMaterialCompiled(UMaterialInterface* InMat
 	}
 }
 
+bool FMaterialInterfaceImageViewer::OnOwnerCleanup(UObject* InObject)
+{
+	if (InObject && MaterialSettings.MaterialInterface && 
+		(MaterialSettings.MaterialInterface == InObject || MaterialSettings.MaterialInterface->IsIn(InObject)))
+	{
+		SampleCache.Reset();
+
+		if (Brush.IsValid())
+		{
+			Brush->SetResourceObject(nullptr);
+		}
+
+		return true;
+	}
+
+	return false;
+}
+
 } // UE::MediaViewer::Private
 
 #undef LOCTEXT_NAMESPACE
diff --git a/Engine/Plugins/Media/MediaViewer/Source/MediaViewer/Private/ImageViewers/MaterialInterfaceImageViewer.h b/Engine/Plugins/Media/MediaViewer/Source/MediaViewer/Private/ImageViewers/MaterialInterfaceImageViewer.h
--- a/Engine/Plugins/Media/MediaViewer/Source/MediaViewer/Private/ImageViewers/MaterialInterfaceImageViewer.h
+++ b/Engine/Plugins/Media/MediaViewer/Source/MediaViewer/Private/ImageViewers/MaterialInterfaceImageViewer.h
@@ -5,7 +5,6 @@
 #include "ImageViewer/IMediaImageViewerFactory.h"
 #include "ImageViewer/MediaImageViewer.h"
 #include "Library/MediaViewerLibraryItem.h"
-#include "UObject/GCObject.h"
 
 #include "ImageViewers/TextureSampleCache.h"
 #include "Templates/SharedPointer.h"
@@ -38,7 +37,7 @@ struct FMaterialInterfaceImageViewerSettings
 namespace UE::MediaViewer::Private
 {
 
-class FMaterialInterfaceImageViewer : public FMediaImageViewer, public FGCObject
+class FMaterialInterfaceImageViewer : public FMediaImageViewer
 {
 public:
 	struct FFactory : public IMediaImageViewerFactory
@@ -93,6 +92,7 @@ public:
 
 	//~ Begin FMediaImageViewer
 	virtual TSharedPtr<FMediaViewerLibraryItem> CreateLibraryItem() const override;
+	virtual bool OnOwnerCleanup(UObject* InObject) override;
 	virtual TOptional<TVariant<FColor, FLinearColor>> GetPixelColor(const FIntPoint& InPixelCoords, int32 InMipLevel) const override;
 	virtual TSharedPtr<FStructOnScope> GetCustomSettingsOnScope() const override;
 	virtual void PaintImage(FMediaImageSlatePaintParams& InPaintParams, const FMediaImageSlatePaintGeometry& InPaintGeometry) override;
@@ -104,7 +104,7 @@ public:
 
 	//~ Begin FGCObject
 	virtual FString GetReferencerName() const override;
-	virtual void AddReferencedObjects(FReferenceCollector& Collector) override;
+	virtual void AddReferencedObjects(FReferenceCollector& InCollector) override;
 	//~ End FGCObject
 
 protected:
diff --git a/Engine/Plugins/Media/MediaViewer/Source/MediaViewer/Private/ImageViewers/MediaImageViewer.cpp b/Engine/Plugins/Media/MediaViewer/Source/MediaViewer/Private/ImageViewers/MediaImageViewer.cpp
--- a/Engine/Plugins/Media/MediaViewer/Source/MediaViewer/Private/ImageViewers/MediaImageViewer.cpp
+++ b/Engine/Plugins/Media/MediaViewer/Source/MediaViewer/Private/ImageViewers/MediaImageViewer.cpp
@@ -5,6 +5,7 @@
 #include "Engine/Canvas.h"
 #include "Engine/Texture.h"
 #include "Engine/TextureRenderTarget2D.h"
+#include "Engine/World.h"
 #include "GameFramework/Actor.h"
 #include "Materials/Material.h"
 #include "Materials/MaterialInstanceDynamic.h"
@@ -23,7 +24,7 @@ namespace UE::MediaViewer
 const FSlateColorBrush FMediaImageViewer::BackgroundColorBrush(FLinearColor::White);
 constexpr const TCHAR* MipRenderMaterial = TEXT("/Script/Engine.Material'/MediaViewer/M_TextureMip.M_TextureMip'");
 
-MEDIAVIEWER_API FText FMediaImageViewer::GetObjectDisplayName(const UObject* InObject)
+FText FMediaImageViewer::GetObjectDisplayName(const UObject* InObject)
 {
 	if (!InObject)
 	{
@@ -62,6 +63,28 @@ void FMediaImageViewer::NotifyPostChange(const FPropertyChangedEvent& InProperty
 	InitBackgroundTexture();
 }
 
+void FMediaImageViewer::AddReferencedObjects(FReferenceCollector& InCollector)
+{
+	if (Brush.IsValid())
+	{
+		if (UObject* Object = Brush->GetResourceObject())
+		{
+			InCollector.AddPropertyReferencesWithStructARO(
+				FSlateBrush::StaticStruct(),
+				Brush.Get()
+			);
+		}
+	}
+	
+	if (UObject* Object = BackgroundImageBrush.GetResourceObject())
+	{
+		InCollector.AddPropertyReferencesWithStructARO(
+			FSlateBrush::StaticStruct(),
+			&BackgroundImageBrush
+		);
+	}
+}
+
 FVector2D FMediaImageViewer::GetViewerCenter(const FVector2D& InViewerSize) const
 {
 	return InViewerSize * 0.5;
@@ -270,7 +293,7 @@ void FMediaImageViewer::PaintImage(FMediaImageSlatePaintParams& InPaintParams, c
 		}
 	}
 
-	if (!BrushToUse.IsValid())
+	if (!BrushToUse.IsValid() || !BrushToUse->GetResourceObject())
 	{
 		return;
 	}
diff --git a/Engine/Plugins/Media/MediaViewer/Source/MediaViewer/Private/ImageViewers/MediaSourceImageViewer.cpp b/Engine/Plugins/Media/MediaViewer/Source/MediaViewer/Private/ImageViewers/MediaSourceImageViewer.cpp
--- a/Engine/Plugins/Media/MediaViewer/Source/MediaViewer/Private/ImageViewers/MediaSourceImageViewer.cpp
+++ b/Engine/Plugins/Media/MediaViewer/Source/MediaViewer/Private/ImageViewers/MediaSourceImageViewer.cpp
@@ -149,11 +149,11 @@ FString FMediaSourceImageViewer::FItem::GetReferencerName() const
 	return ReferencerName;
 }
 
-void FMediaSourceImageViewer::FItem::AddReferencedObjects(FReferenceCollector& Collector)
+void FMediaSourceImageViewer::FItem::AddReferencedObjects(FReferenceCollector& InCollector)
 {
 	if (Texture)
 	{
-		Collector.AddReferencedObject(Texture);
+		InCollector.AddReferencedObject(Texture);
 	}
 }
 
@@ -436,16 +436,18 @@ TSharedPtr<FStructOnScope> FMediaSourceImageViewer::GetCustomSettingsOnScope() c
 	);
 }
 
-void FMediaSourceImageViewer::AddReferencedObjects(FReferenceCollector& Collector)
+void FMediaSourceImageViewer::AddReferencedObjects(FReferenceCollector& InCollector)
 {
-	Collector.AddPropertyReferencesWithStructARO(
+	FMediaImageViewer::AddReferencedObjects(InCollector);
+
+	InCollector.AddPropertyReferencesWithStructARO(
 		FMediaSourceImageViewerSettings::StaticStruct(),
 		&MediaSourceSettings
 	);
 
 	if (MediaStream)
 	{
-		Collector.AddReferencedObject(MediaStream);
+		InCollector.AddReferencedObject(MediaStream);
 	}
 }
 
@@ -509,6 +511,37 @@ void FMediaSourceImageViewer::AddPlayerName(const TSharedRef<SHorizontalBox>& In
 		];
 }
 
+bool FMediaSourceImageViewer::OnOwnerCleanup(UObject* InObject)
+{
+	bool bInvalidate = false;
+
+	if (InObject)
+	{
+		if (MediaSourceSettings.MediaSource && 
+			(MediaSourceSettings.MediaSource == InObject || MediaSourceSettings.MediaSource->IsIn(InObject)))
+		{
+			bInvalidate = true;
+		}
+
+		if (MediaStream && (MediaStream == InObject || MediaStream->IsIn(InObject)))
+		{
+			bInvalidate = true;
+		}
+	}
+
+	if (bInvalidate)
+	{
+		SampleCache.Reset();
+
+		if (Brush.IsValid())
+		{
+			Brush->SetResourceObject(nullptr);
+		}
+	}
+
+	return bInvalidate;
+}
+
 FString FMediaSourceImageViewer::GetReferencerName() const
 {
 	static const FString ReferencerName = TEXT("FMediaSourceImageViewer");
diff --git a/Engine/Plugins/Media/MediaViewer/Source/MediaViewer/Private/ImageViewers/MediaSourceImageViewer.h b/Engine/Plugins/Media/MediaViewer/Source/MediaViewer/Private/ImageViewers/MediaSourceImageViewer.h
--- a/Engine/Plugins/Media/MediaViewer/Source/MediaViewer/Private/ImageViewers/MediaSourceImageViewer.h
+++ b/Engine/Plugins/Media/MediaViewer/Source/MediaViewer/Private/ImageViewers/MediaSourceImageViewer.h
@@ -5,7 +5,6 @@
 #include "ImageViewer/IMediaImageViewerFactory.h"
 #include "ImageViewer/MediaImageViewer.h"
 #include "Library/MediaViewerLibraryItem.h"
-#include "UObject/GCObject.h"
 
 #include "Engine/TextureDefines.h"
 #include "ImageViewers/TextureSampleCache.h"
@@ -35,7 +34,7 @@ struct FMediaSourceImageViewerSettings
 namespace UE::MediaViewer::Private
 {
 
-class FMediaSourceImageViewer : public FMediaImageViewer, public FGCObject
+class FMediaSourceImageViewer : public FMediaImageViewer
 {
 public:
 	struct FFactory : public IMediaImageViewerFactory
@@ -72,7 +71,7 @@ public:
 
 		//~ Begin FGCObject
 		virtual FString GetReferencerName() const override;
-		virtual void AddReferencedObjects(FReferenceCollector& Collector) override;
+		virtual void AddReferencedObjects(FReferenceCollector& InCollector) override;
 		//~ End FGCObject
 	};
 
@@ -125,6 +124,7 @@ public:
 
 	//~ Begin FMediaImageViewer
 	virtual TSharedPtr<FMediaViewerLibraryItem> CreateLibraryItem() const override;
+	virtual bool OnOwnerCleanup(UObject* InObject) override;
 	virtual TOptional<TVariant<FColor, FLinearColor>> GetPixelColor(const FIntPoint& InPixelCoords, int32 InMipLevel) const override;
 	virtual TSharedPtr<FStructOnScope> GetCustomSettingsOnScope() const override;
 	virtual TSharedPtr<SWidget> GetOverlayWidget(EMediaImageViewerPosition InPosition, const TSharedPtr<FMediaViewerDelegates>& InDelegates) override;
@@ -133,7 +133,7 @@ public:
 
 	//~ Begin FGCObject
 	virtual FString GetReferencerName() const override;
-	virtual void AddReferencedObjects(FReferenceCollector& Collector) override;
+	virtual void AddReferencedObjects(FReferenceCollector& InCollector) override;
 	//~ End FGCObject
 
 protected:
diff --git a/Engine/Plugins/Media/MediaViewer/Source/MediaViewer/Private/ImageViewers/MediaTextureImageViewer.cpp b/Engine/Plugins/Media/MediaViewer/Source/MediaViewer/Private/ImageViewers/MediaTextureImageViewer.cpp
--- a/Engine/Plugins/Media/MediaViewer/Source/MediaViewer/Private/ImageViewers/MediaTextureImageViewer.cpp
+++ b/Engine/Plugins/Media/MediaViewer/Source/MediaViewer/Private/ImageViewers/MediaTextureImageViewer.cpp
@@ -272,14 +272,34 @@ TSharedPtr<FStructOnScope> FMediaTextureImageViewer::GetCustomSettingsOnScope()
 	);
 }
 
-void FMediaTextureImageViewer::AddReferencedObjects(FReferenceCollector& Collector)
+void FMediaTextureImageViewer::AddReferencedObjects(FReferenceCollector& InCollector)
 {
-	Collector.AddPropertyReferencesWithStructARO(
+	FMediaImageViewer::AddReferencedObjects(InCollector);
+
+	InCollector.AddPropertyReferencesWithStructARO(
 		FMediaTextureImageViewerSettings::StaticStruct(),
 		&MediaTextureSettings
 	);
 }
 
+bool FMediaTextureImageViewer::OnOwnerCleanup(UObject* InObject)
+{
+	if (InObject && MediaTextureSettings.MediaTexture && 
+		(MediaTextureSettings.MediaTexture == InObject || MediaTextureSettings.MediaTexture->IsIn(InObject)))
+	{
+		SampleCache.Reset();
+
+		if (Brush.IsValid())
+		{
+			Brush->SetResourceObject(nullptr);
+		}
+
+		return true;
+	}
+
+	return false;
+}
+
 FString FMediaTextureImageViewer::GetReferencerName() const
 {
 	static const FString ReferencerName = TEXT("FMediaTextureImageViewer");
diff --git a/Engine/Plugins/Media/MediaViewer/Source/MediaViewer/Private/ImageViewers/MediaTextureImageViewer.h b/Engine/Plugins/Media/MediaViewer/Source/MediaViewer/Private/ImageViewers/MediaTextureImageViewer.h
--- a/Engine/Plugins/Media/MediaViewer/Source/MediaViewer/Private/ImageViewers/MediaTextureImageViewer.h
+++ b/Engine/Plugins/Media/MediaViewer/Source/MediaViewer/Private/ImageViewers/MediaTextureImageViewer.h
@@ -5,7 +5,6 @@
 #include "ImageViewer/IMediaImageViewerFactory.h"
 #include "ImageViewer/MediaImageViewer.h"
 #include "Library/MediaViewerLibraryItem.h"
-#include "UObject/GCObject.h"
 
 #include "ImageViewers/TextureSampleCache.h"
 #include "Templates/SharedPointer.h"
@@ -29,7 +28,7 @@ struct FMediaTextureImageViewerSettings
 namespace UE::MediaViewer::Private
 {
 
-class FMediaTextureImageViewer : public FMediaImageViewer, public FGCObject
+class FMediaTextureImageViewer : public FMediaImageViewer
 {
 public:
 	struct FFactory : public IMediaImageViewerFactory
@@ -81,13 +80,14 @@ public:
 
 	//~ Begin FMediaImageViewer
 	virtual TSharedPtr<FMediaViewerLibraryItem> CreateLibraryItem() const override;
+	virtual bool OnOwnerCleanup(UObject* InObject) override;
 	virtual TOptional<TVariant<FColor, FLinearColor>> GetPixelColor(const FIntPoint& InPixelCoords, int32 InMipLevel) const override;
 	virtual TSharedPtr<FStructOnScope> GetCustomSettingsOnScope() const override;
 	//~ End FMediaImageViewer
 
 	//~ Begin FGCObject
 	virtual FString GetReferencerName() const override;
-	virtual void AddReferencedObjects(FReferenceCollector& Collector) override;
+	virtual void AddReferencedObjects(FReferenceCollector& InCollector) override;
 	//~ End FGCObject
 
 protected:
diff --git a/Engine/Plugins/Media/MediaViewer/Source/MediaViewer/Private/ImageViewers/NullImageViewer.cpp b/Engine/Plugins/Media/MediaViewer/Source/MediaViewer/Private/ImageViewers/NullImageViewer.cpp
--- a/Engine/Plugins/Media/MediaViewer/Source/MediaViewer/Private/ImageViewers/NullImageViewer.cpp
+++ b/Engine/Plugins/Media/MediaViewer/Source/MediaViewer/Private/ImageViewers/NullImageViewer.cpp
@@ -33,6 +33,12 @@ TOptional<TVariant<FColor, FLinearColor>> FNullImageViewer::GetPixelColor(const
 	return {};
 }
 
+FString FNullImageViewer::GetReferencerName() const
+{
+	static const FString ReferencerName = TEXT("FNullImageViewer");
+	return ReferencerName;
+}
+
 TSharedPtr<FMediaViewerLibraryItem> FNullImageViewer::CreateLibraryItem() const
 {
 	return nullptr;
diff --git a/Engine/Plugins/Media/MediaViewer/Source/MediaViewer/Private/ImageViewers/NullImageViewer.h b/Engine/Plugins/Media/MediaViewer/Source/MediaViewer/Private/ImageViewers/NullImageViewer.h
--- a/Engine/Plugins/Media/MediaViewer/Source/MediaViewer/Private/ImageViewers/NullImageViewer.h
+++ b/Engine/Plugins/Media/MediaViewer/Source/MediaViewer/Private/ImageViewers/NullImageViewer.h
@@ -19,6 +19,10 @@ public:
 	virtual TOptional<TVariant<FColor, FLinearColor>> GetPixelColor(const FIntPoint& InPixelCoords, int32 InMipLevel) const override;
 	//~ End FMediaImageViewer
 
+	//~ Begin FGCObject
+	virtual FString GetReferencerName() const override;
+	//~ End FGCObject
+
 protected:
 	//~ Begin FMediaImageViewer
 	virtual TSharedPtr<FMediaViewerLibraryItem> CreateLibraryItem() const override;
diff --git a/Engine/Plugins/Media/MediaViewer/Source/MediaViewer/Private/ImageViewers/Texture2DImageViewer.cpp b/Engine/Plugins/Media/MediaViewer/Source/MediaViewer/Private/ImageViewers/Texture2DImageViewer.cpp
--- a/Engine/Plugins/Media/MediaViewer/Source/MediaViewer/Private/ImageViewers/Texture2DImageViewer.cpp
+++ b/Engine/Plugins/Media/MediaViewer/Source/MediaViewer/Private/ImageViewers/Texture2DImageViewer.cpp
@@ -206,11 +206,11 @@ FString FTexture2DImageViewer::FItem::GetReferencerName() const
 	return ReferencerName;
 }
 
-void FTexture2DImageViewer::FItem::AddReferencedObjects(FReferenceCollector& Collector)
+void FTexture2DImageViewer::FItem::AddReferencedObjects(FReferenceCollector& InCollector)
 {
 	if (Texture)
 	{
-		Collector.AddReferencedObject(Texture);
+		InCollector.AddReferencedObject(Texture);
 	}
 }
 
@@ -296,9 +296,11 @@ TSharedPtr<FStructOnScope> FTexture2DImageViewer::GetCustomSettingsOnScope() con
 	);
 }
 
-void FTexture2DImageViewer::AddReferencedObjects(FReferenceCollector& Collector)
+void FTexture2DImageViewer::AddReferencedObjects(FReferenceCollector& InCollector)
 {
-	Collector.AddPropertyReferencesWithStructARO(
+	FMediaImageViewer::AddReferencedObjects(InCollector);
+
+	InCollector.AddPropertyReferencesWithStructARO(
 		FTexture2DImageViewerSettings::StaticStruct(),
 		&TextureSettings
 	);
@@ -329,6 +331,24 @@ TOptional<FIntPoint> FTexture2DImageViewer::GetTextureSize() const
 	return FIntPoint(TextureSettings.Texture->GetSurfaceWidth(), TextureSettings.Texture->GetSurfaceHeight());
 }
 
+bool FTexture2DImageViewer::OnOwnerCleanup(UObject* InObject)
+{
+	if (InObject && TextureSettings.Texture && 
+		(TextureSettings.Texture == InObject || TextureSettings.Texture->IsIn(InObject)))
+	{
+		SampleCache.Reset();
+
+		if (Brush.IsValid())
+		{
+			Brush->SetResourceObject(nullptr);
+		}
+
+		return true;
+	}
+
+	return false;
+}
+
 void FTexture2DImageViewer::PaintImage(FMediaImageSlatePaintParams& InPaintParams, const FMediaImageSlatePaintGeometry& InPaintGeometry)
 {
 	if (!bValidImageSize)
diff --git a/Engine/Plugins/Media/MediaViewer/Source/MediaViewer/Private/ImageViewers/Texture2DImageViewer.h b/Engine/Plugins/Media/MediaViewer/Source/MediaViewer/Private/ImageViewers/Texture2DImageViewer.h
--- a/Engine/Plugins/Media/MediaViewer/Source/MediaViewer/Private/ImageViewers/Texture2DImageViewer.h
+++ b/Engine/Plugins/Media/MediaViewer/Source/MediaViewer/Private/ImageViewers/Texture2DImageViewer.h
@@ -5,7 +5,6 @@
 #include "ImageViewer/IMediaImageViewerFactory.h"
 #include "ImageViewer/MediaImageViewer.h"
 #include "Library/MediaViewerLibraryItem.h"
-#include "UObject/GCObject.h"
 
 #include "ImageViewers/TextureSampleCache.h"
 #include "Math/MathFwd.h"
@@ -28,7 +27,7 @@ struct FTexture2DImageViewerSettings
 namespace UE::MediaViewer::Private
 {
 
-class FTexture2DImageViewer : public FMediaImageViewer, public FGCObject
+class FTexture2DImageViewer : public FMediaImageViewer
 {
 public:
 	struct FFactory : public IMediaImageViewerFactory
@@ -74,7 +73,7 @@ public:
 
 		//~ Begin FGCObject
 		virtual FString GetReferencerName() const override;
-		virtual void AddReferencedObjects(FReferenceCollector& Collector) override;
+		virtual void AddReferencedObjects(FReferenceCollector& InCollector) override;
 		//~ End FGCObject
 	};
 
@@ -87,13 +86,14 @@ public:
 
 	//~ Begin FMediaImageViewer
 	virtual TSharedPtr<FMediaViewerLibraryItem> CreateLibraryItem() const override;
+	virtual bool OnOwnerCleanup(UObject* InObject) override;
 	virtual TOptional<TVariant<FColor, FLinearColor>> GetPixelColor(const FIntPoint& InPixelCoords, int32 InMipLevel) const override;
 	virtual TSharedPtr<FStructOnScope> GetCustomSettingsOnScope() const override;
 	//~ End FMediaImageViewer
 
 	//~ Begin FGCObject
 	virtual FString GetReferencerName() const override;
-	virtual void AddReferencedObjects(FReferenceCollector& Collector) override;
+	virtual void AddReferencedObjects(FReferenceCollector& InCollector) override;
 	//~ End FGCObject
 
 protected:
diff --git a/Engine/Plugins/Media/MediaViewer/Source/MediaViewer/Private/ImageViewers/TextureRenderTarget2DImageViewer.cpp b/Engine/Plugins/Media/MediaViewer/Source/MediaViewer/Private/ImageViewers/TextureRenderTarget2DImageViewer.cpp
--- a/Engine/Plugins/Media/MediaViewer/Source/MediaViewer/Private/ImageViewers/TextureRenderTarget2DImageViewer.cpp
+++ b/Engine/Plugins/Media/MediaViewer/Source/MediaViewer/Private/ImageViewers/TextureRenderTarget2DImageViewer.cpp
@@ -266,14 +266,34 @@ TSharedPtr<FStructOnScope> FTextureRenderTarget2DImageViewer::GetCustomSettingsO
 	);
 }
 
-void FTextureRenderTarget2DImageViewer::AddReferencedObjects(FReferenceCollector& Collector)
+void FTextureRenderTarget2DImageViewer::AddReferencedObjects(FReferenceCollector& InCollector)
 {
-	Collector.AddPropertyReferencesWithStructARO(
+	FMediaImageViewer::AddReferencedObjects(InCollector);
+
+	InCollector.AddPropertyReferencesWithStructARO(
 		FTextureRenderTarget2DImageViewerSettings::StaticStruct(),
 		&RenderTargetSettings
 	);
 }
 
+bool FTextureRenderTarget2DImageViewer::OnOwnerCleanup(UObject* InObject)
+{
+	if (InObject && RenderTargetSettings.RenderTarget && 
+		(RenderTargetSettings.RenderTarget == InObject || RenderTargetSettings.RenderTarget->IsIn(InObject)))
+	{
+		SampleCache.Reset();
+
+		if (Brush.IsValid())
+		{
+			Brush->SetResourceObject(nullptr);
+		}
+
+		return true;
+	}
+
+	return false;
+}
+
 FString FTextureRenderTarget2DImageViewer::GetReferencerName() const
 {
 	static const FString ReferencerName = TEXT("FTextureRenderTarget2DImageViewer");
diff --git a/Engine/Plugins/Media/MediaViewer/Source/MediaViewer/Private/ImageViewers/TextureRenderTarget2DImageViewer.h b/Engine/Plugins/Media/MediaViewer/Source/MediaViewer/Private/ImageViewers/TextureRenderTarget2DImageViewer.h
--- a/Engine/Plugins/Media/MediaViewer/Source/MediaViewer/Private/ImageViewers/TextureRenderTarget2DImageViewer.h
+++ b/Engine/Plugins/Media/MediaViewer/Source/MediaViewer/Private/ImageViewers/TextureRenderTarget2DImageViewer.h
@@ -5,7 +5,6 @@
 #include "ImageViewer/IMediaImageViewerFactory.h"
 #include "ImageViewer/MediaImageViewer.h"
 #include "Library/MediaViewerLibraryItem.h"
-#include "UObject/GCObject.h"
 
 #include "ImageViewers/TextureSampleCache.h"
 #include "Templates/SharedPointer.h"
@@ -26,7 +25,7 @@ struct FTextureRenderTarget2DImageViewerSettings
 namespace UE::MediaViewer::Private
 {
 
-class FTextureRenderTarget2DImageViewer : public FMediaImageViewer, public FGCObject
+class FTextureRenderTarget2DImageViewer : public FMediaImageViewer
 {
 public:
 	struct FFactory : public IMediaImageViewerFactory
@@ -78,13 +77,14 @@ public:
 
 	//~ Begin FMediaImageViewer
 	virtual TSharedPtr<FMediaViewerLibraryItem> CreateLibraryItem() const override;
+	virtual bool OnOwnerCleanup(UObject* InObject) override;
 	virtual TOptional<TVariant<FColor, FLinearColor>> GetPixelColor(const FIntPoint& InPixelCoords, int32 InMipLevel) const override;
 	virtual TSharedPtr<FStructOnScope> GetCustomSettingsOnScope() const override;
 	//~ End FMediaImageViewer
 
 	//~ Begin FGCObject
 	virtual FString GetReferencerName() const override;
-	virtual void AddReferencedObjects(FReferenceCollector& Collector) override;
+	virtual void AddReferencedObjects(FReferenceCollector& InCollector) override;
 	//~ End FGCObject
 
 protected:
diff --git a/Engine/Plugins/Media/MediaViewer/Source/MediaViewer/Private/ImageViewers/TextureSampleCache.cpp b/Engine/Plugins/Media/MediaViewer/Source/MediaViewer/Private/ImageViewers/TextureSampleCache.cpp
--- a/Engine/Plugins/Media/MediaViewer/Source/MediaViewer/Private/ImageViewers/TextureSampleCache.cpp
+++ b/Engine/Plugins/Media/MediaViewer/Source/MediaViewer/Private/ImageViewers/TextureSampleCache.cpp
@@ -23,7 +23,7 @@ FTextureSampleCache::FTextureSampleCache()
 }
 
 FTextureSampleCache::FTextureSampleCache(TNotNull<UTexture*> InTexture, EPixelFormat InPixelFormat)
-	: Texture(InTexture)
+	: TextureWeak(InTexture)
 	, SampleCS(FCriticalSection())
 	, RenderTarget(nullptr)
 {
@@ -31,11 +31,13 @@ FTextureSampleCache::FTextureSampleCache(TNotNull<UTexture*> InTexture, EPixelFo
 
 bool FTextureSampleCache::IsValid() const
 {
-	return ::IsValid(Texture);
+	return TextureWeak.IsValid();
 }
 
 bool FTextureSampleCache::NeedsUpdate(const FIntPoint& InPixelCoordinates, const TOptional<FTimespan>& InTime) const
 {
+	UTexture* Texture = TextureWeak.Get();
+
 	if (!Texture)
 	{
 		return false;
@@ -55,6 +57,8 @@ bool FTextureSampleCache::NeedsUpdate(const FIntPoint& InPixelCoordinates, const
 
 const FLinearColor* FTextureSampleCache::GetPixelColor(const FIntPoint& InPixelCoordinates, const TOptional<FTimespan>& InTime)
 {
+	UTexture* Texture = TextureWeak.Get();
+
 	if (!Texture)
 	{
 		return nullptr;
@@ -106,7 +110,7 @@ FTextureSampleCache& FTextureSampleCache::operator=(const FTextureSampleCache& I
 	FScopeLock MyLock(&SampleCS);
 	FScopeLock OtherLock(&InOther.SampleCS);
 
-	Texture = InOther.Texture;
+	TextureWeak = InOther.TextureWeak;
 	bDirty = InOther.bDirty;
 	PixelColorSample = InOther.PixelColorSample;
 	
@@ -119,11 +123,11 @@ FString FTextureSampleCache::GetReferencerName() const
 	return ReferencerName;
 }
 
-void FTextureSampleCache::AddReferencedObjects(FReferenceCollector& Collector)
+void FTextureSampleCache::AddReferencedObjects(FReferenceCollector& InCollector)
 {
 	if (RenderTarget)
 	{
-		Collector.AddReferencedObject(RenderTarget);
+		InCollector.AddReferencedObject(RenderTarget);
 	}
 }
 
@@ -181,6 +185,13 @@ void FTextureSampleCache::SetPixelColor_RHI(const FIntPoint& InPixelCoordinates,
 
 void FTextureSampleCache::SetPixelColor_RenderTarget(const FIntPoint& InPixelCoordinates, const TOptional<FTimespan>& InTime)
 {
+	UTexture* Texture = TextureWeak.Get();
+
+	if (!Texture)
+	{
+		return;
+	}
+
 	if (!RenderTarget)
 	{
 		RenderTarget = FMediaViewerUtils::CreateRenderTarget(
diff --git a/Engine/Plugins/Media/MediaViewer/Source/MediaViewer/Private/ImageViewers/TextureSampleCache.h b/Engine/Plugins/Media/MediaViewer/Source/MediaViewer/Private/ImageViewers/TextureSampleCache.h
--- a/Engine/Plugins/Media/MediaViewer/Source/MediaViewer/Private/ImageViewers/TextureSampleCache.h
+++ b/Engine/Plugins/Media/MediaViewer/Source/MediaViewer/Private/ImageViewers/TextureSampleCache.h
@@ -11,6 +11,7 @@
 #include "Misc/NotNull.h"
 #include "Misc/Optional.h"
 #include "Misc/Timespan.h"
+#include "UObject/WeakObjectPtr.h"
 
 class UTexture;
 class UTextureRenderTarget2D;
@@ -37,7 +38,7 @@ public:
 
 	//~ Begin FGCObject
 	virtual FString GetReferencerName() const override;
-	virtual void AddReferencedObjects(FReferenceCollector& Collector) override;
+	virtual void AddReferencedObjects(FReferenceCollector& InCollector) override;
 	//~ End FGCObject
 
 protected:
@@ -48,7 +49,7 @@ protected:
 		FLinearColor Color;
 	};
 
-	UTexture* Texture = nullptr;
+	TWeakObjectPtr<UTexture> TextureWeak;
 	bool bDirty = true;
 	TOptional<FPixelColorSample> PixelColorSample;
 	mutable FCriticalSection SampleCS;
diff --git a/Engine/Plugins/Media/MediaViewer/Source/MediaViewer/Private/ImageViewers/ViewportImageViewer.cpp b/Engine/Plugins/Media/MediaViewer/Source/MediaViewer/Private/ImageViewers/ViewportImageViewer.cpp
--- a/Engine/Plugins/Media/MediaViewer/Source/MediaViewer/Private/ImageViewers/ViewportImageViewer.cpp
+++ b/Engine/Plugins/Media/MediaViewer/Source/MediaViewer/Private/ImageViewers/ViewportImageViewer.cpp
@@ -93,9 +93,11 @@ FString FViewportImageViewer::GetReferencerName() const
 	return ReferencerName;
 }
 
-void FViewportImageViewer::AddReferencedObjects(FReferenceCollector& Collector)
+void FViewportImageViewer::AddReferencedObjects(FReferenceCollector& InCollector)
 {
-	Collector.AddPropertyReferencesWithStructARO(
+	FMediaImageViewer::AddReferencedObjects(InCollector);
+
+	InCollector.AddPropertyReferencesWithStructARO(
 		FViewportImageViewerSettings::StaticStruct(),
 		&ViewportSettings
 	);
diff --git a/Engine/Plugins/Media/MediaViewer/Source/MediaViewer/Private/ImageViewers/ViewportImageViewer.h b/Engine/Plugins/Media/MediaViewer/Source/MediaViewer/Private/ImageViewers/ViewportImageViewer.h
--- a/Engine/Plugins/Media/MediaViewer/Source/MediaViewer/Private/ImageViewers/ViewportImageViewer.h
+++ b/Engine/Plugins/Media/MediaViewer/Source/MediaViewer/Private/ImageViewers/ViewportImageViewer.h
@@ -4,7 +4,6 @@
 
 #include "ImageViewer/IMediaImageViewerFactory.h"
 #include "ImageViewer/MediaImageViewer.h"
-#include "UObject/GCObject.h"
 
 #include "Delegates/IDelegateInstance.h"
 #include "ImageViewers/TextureSampleCache.h"
@@ -33,7 +32,7 @@ struct FViewportImageViewerSettings
 namespace UE::MediaViewer::Private
 {
 
-class FViewportImageViewer : public FMediaImageViewer, public FGCObject
+class FViewportImageViewer : public FMediaImageViewer
 {
 public:
 	FViewportImageViewer(const FMediaImageViewerInfo& InImageInfo);
@@ -48,7 +47,7 @@ public:
 
 	//~ Begin FGCObject
 	virtual FString GetReferencerName() const override;
-	virtual void AddReferencedObjects(FReferenceCollector& Collector) override;
+	virtual void AddReferencedObjects(FReferenceCollector& InCollector) override;
 	//~ End FGCObject
 
 protected:
diff --git a/Engine/Plugins/Media/MediaViewer/Source/MediaViewer/Private/Widgets/MediaViewerLibraryGroupItem.h b/Engine/Plugins/Media/MediaViewer/Source/MediaViewer/Private/Widgets/MediaViewerLibraryGroupItem.h
--- a/Engine/Plugins/Media/MediaViewer/Source/MediaViewer/Private/Widgets/MediaViewerLibraryGroupItem.h
+++ b/Engine/Plugins/Media/MediaViewer/Source/MediaViewer/Private/Widgets/MediaViewerLibraryGroupItem.h
@@ -28,7 +28,7 @@ public:
 	using SetKeyFuncs = DefaultKeyFuncs<UE::MediaViewer::IMediaViewerLibrary::FGroupItem>;
 
 	template<typename U>
-	static void AddReferencedObjects(FReferenceCollector& Collector,
+	static void AddReferencedObjects(FReferenceCollector& InCollector,
 		TArray<UE::MediaViewer::IMediaViewerLibrary::FGroupItem>& ItemsWithGeneratedWidgets,
 		TSet<UE::MediaViewer::IMediaViewerLibrary::FGroupItem>& SelectedItems,
 		TMap<const U*, UE::MediaViewer::IMediaViewerLibrary::FGroupItem>& WidgetToItemMap)
diff --git a/Engine/Plugins/Media/MediaViewer/Source/MediaViewer/Private/Widgets/SMediaViewer.cpp b/Engine/Plugins/Media/MediaViewer/Source/MediaViewer/Private/Widgets/SMediaViewer.cpp
--- a/Engine/Plugins/Media/MediaViewer/Source/MediaViewer/Private/Widgets/SMediaViewer.cpp
+++ b/Engine/Plugins/Media/MediaViewer/Source/MediaViewer/Private/Widgets/SMediaViewer.cpp
@@ -7,6 +7,7 @@
 #include "Brushes/SlateImageBrush.h"
 #include "CanvasTypes.h"
 #include "ContentBrowserModule.h"
+#include "Editor.h"
 #include "Engine/Canvas.h"
 #include "Engine/Engine.h"
 #include "Engine/Texture2D.h"
@@ -29,6 +30,7 @@
 #include "Sidebar/SSidebar.h"
 #include "Sidebar/SSidebarContainer.h"
 #include "TextureResource.h"
+#include "UObject/Package.h"
 #include "UObject/StrongObjectPtr.h"
 #include "Widgets/Colors/SColorBlock.h"
 #include "Widgets/Layout/SBox.h"
@@ -57,10 +59,16 @@ SMediaViewer::SMediaViewer()
 	, ViewerPosition(FVector2D::ZeroVector)
 	, bInvalidated(true)
 {
+	FWorldDelegates::OnWorldCleanup.AddRaw(this, &SMediaViewer::OnWorldCleanUp);
+	FEditorDelegates::OnPackageDeleted.AddRaw(this, &SMediaViewer::OnPackageDeleted);
+	FEditorDelegates::OnAssetsPreDelete.AddRaw(this, &SMediaViewer::OnAssetsDeleted);
 }
 
 SMediaViewer::~SMediaViewer()
 {
+	FWorldDelegates::OnWorldCleanup.RemoveAll(this);
+	FEditorDelegates::OnPackageDeleted.RemoveAll(this);
+	FEditorDelegates::OnAssetsPreDelete.RemoveAll(this);
 }
 
 void SMediaViewer::PrivateRegisterAttributes(FSlateAttributeDescriptor::FInitializer&)
@@ -539,9 +547,9 @@ FReply SMediaViewer::OnKeyDown(const FGeometry& InMyGeometry, const FKeyEvent& I
 	return SCompoundWidget::OnKeyDown(InMyGeometry, InKeyEvent);
 }
 
-void SMediaViewer::AddReferencedObjects(FReferenceCollector& Collector)
+void SMediaViewer::AddReferencedObjects(FReferenceCollector& InCollector)
 {
-	Collector.AddPropertyReferencesWithStructARO(
+	InCollector.AddPropertyReferencesWithStructARO(
 		FMediaViewerSettings::StaticStruct(),
 		&MediaViewerSettings
 	);
@@ -1396,6 +1404,43 @@ void SMediaViewer::SaveBookmark(int32 InIndex)
 	FSlateNotificationManager::Get().AddNotification(NotificationInfo);
 }
 
+void SMediaViewer::OnOwnerCleanup(UObject* InOuter)
+{
+	const FGuid& NullImageId = FNullImageViewer::GetNullImageViewer()->GetInfo().Id;
+
+	for (int32 Index = 0; Index < static_cast<int32>(EMediaImageViewerPosition::COUNT); ++Index)
+	{
+		if (ImageViewers[Index].IsValid() && ImageViewers[Index]->GetInfo().Id != NullImageId)
+		{
+			if (ImageViewers[Index]->OnOwnerCleanup(InOuter))
+			{
+				ClearImageViewer(static_cast<EMediaImageViewerPosition>(Index));
+			}
+		}
+	}
+}
+
+void SMediaViewer::OnWorldCleanUp(UWorld* InWorld, bool bInSessionEnded, bool bInCleanupResources)
+{
+	if (bInCleanupResources)
+	{
+		OnOwnerCleanup(InWorld);
+	}
+}
+
+void SMediaViewer::OnPackageDeleted(UPackage* InPackage)
+{
+	OnOwnerCleanup(InPackage);
+}
+
+void SMediaViewer::OnAssetsDeleted(const TArray<UObject*>& InAssets)
+{
+	for (UObject* Asset : InAssets)
+	{
+		OnOwnerCleanup(Asset);
+	}
+}
+
 } // UE::MediaViewer::Private
 
 #undef LOCTEXT_NAMESPACE
diff --git a/Engine/Plugins/Media/MediaViewer/Source/MediaViewer/Private/Widgets/SMediaViewer.h b/Engine/Plugins/Media/MediaViewer/Source/MediaViewer/Private/Widgets/SMediaViewer.h
--- a/Engine/Plugins/Media/MediaViewer/Source/MediaViewer/Private/Widgets/SMediaViewer.h
+++ b/Engine/Plugins/Media/MediaViewer/Source/MediaViewer/Private/Widgets/SMediaViewer.h
@@ -17,6 +17,8 @@
 class FUICommandList;
 class SViewportToolBar;
 class UCanvas;
+class UPackage;
+class UWorld;
 struct FMediaViewerState;
 
 namespace UE::MediaViewer
@@ -119,7 +121,7 @@ public:
 
 	//~ Begin FGCObject
 	virtual FString GetReferencerName() const override;
-	virtual void AddReferencedObjects(FReferenceCollector& Collector) override;
+	virtual void AddReferencedObjects(FReferenceCollector& InCollector) override;
 	//~ End FGCObject
 
 protected:
@@ -262,6 +264,14 @@ protected:
 	void OpenBookmark(int32 InIndex);
 	bool CanSaveBookmark() const;
 	void SaveBookmark(int32 InIndex);
+
+	void OnOwnerCleanup(UObject* InOuter);
+
+	void OnWorldCleanUp(UWorld* InWorld, bool bInSessionEnded, bool bInCleanupResources);
+
+	void OnPackageDeleted(UPackage* InPackage);
+
+	void OnAssetsDeleted(const TArray<UObject*>& InAssets);
 };
 
 } // UE::MediaViewer::Private
diff --git a/Engine/Plugins/Media/MediaViewer/Source/MediaViewer/Public/ImageViewer/MediaImageViewer.h b/Engine/Plugins/Media/MediaViewer/Source/MediaViewer/Public/ImageViewer/MediaImageViewer.h
--- a/Engine/Plugins/Media/MediaViewer/Source/MediaViewer/Public/ImageViewer/MediaImageViewer.h
+++ b/Engine/Plugins/Media/MediaViewer/Source/MediaViewer/Public/ImageViewer/MediaImageViewer.h
@@ -16,6 +16,7 @@
 #include "Misc/Guid.h"
 #include "Misc/Optional.h"
 #include "Rendering/RenderingCommon.h"
+#include "UObject/GCObject.h"
 #include "UObject/SoftObjectPtr.h"
 #include "UObject/StructOnScope.h"
 
@@ -32,6 +33,7 @@ class FWidgetStyle;
 class SWidget;
 class UCanvas;
 class UMaterialInstanceDynamic;
+class UObject;
 class UTexture;
 class UTextureRenderTarget2D;
 struct FFloatRange;
@@ -216,7 +218,7 @@ struct FMediaImageCanvasPaintGeometry
 /**
  * Image resource for the AB Image Viewer representing a source that can be displayed as a 2d image.
  */
-class FMediaImageViewer : public FNotifyHook, public TSharedFromThis<FMediaImageViewer>
+class FMediaImageViewer : public FNotifyHook, public FGCObject, public TSharedFromThis<FMediaImageViewer>
 {
 protected:
 	struct FPrivateToken {};
@@ -231,7 +233,7 @@ public:
 
 	MEDIAVIEWER_API FMediaImageViewer(const FMediaImageViewerInfo& InImageInfo);
 
-	virtual ~FMediaImageViewer() = default;
+	virtual ~FMediaImageViewer() override = default;
 
 	const FMediaImageViewerInfo& GetInfo() const { return ImageInfo; }
 
@@ -335,10 +337,20 @@ public:
 	/** Returns the size of the image, including things like scale. */
 	MEDIAVIEWER_API FVector2D GetPaintSize() const;
 
+	/** Called when a potential outer has been removed. Returns true if this is no longer valid. */
+	virtual bool OnOwnerCleanup(UObject* InObject) 
+	{
+		return false;
+	}
+
 	//~ Begin FNotifyHook
 	virtual void NotifyPostChange(const FPropertyChangedEvent& InPropertyChangedEvent, FProperty* InPropertyThatChanged) override;
 	//~ End FNotifyHook
 
+	//~ Begin FGCObject
+	virtual void AddReferencedObjects(FReferenceCollector& InCollector) override;
+	//~ End FGCObject
+
 protected:
 	static const FSlateColorBrush BackgroundColorBrush;
 
