Media Viewer: You can now drag in a single directory or 3 or more images in the same folder to create an image media source.

diff --git a/Engine/Plugins/Media/MediaViewer/Source/MediaViewer/MediaViewer.build.cs b/Engine/Plugins/Media/MediaViewer/Source/MediaViewer/MediaViewer.build.cs
--- a/Engine/Plugins/Media/MediaViewer/Source/MediaViewer/MediaViewer.build.cs
+++ b/Engine/Plugins/Media/MediaViewer/Source/MediaViewer/MediaViewer.build.cs
@@ -27,6 +27,7 @@ public class MediaViewer : ModuleRules
 				"EditorWidgets",
 				"Engine",
 				"ImageCore",
+				"ImgMedia",
 				"InputCore",
 				"LevelEditor",
 				"MediaAssets",
diff --git a/Engine/Plugins/Media/MediaViewer/Source/MediaViewer/Private/Widgets/SMediaViewerDropTarget.cpp b/Engine/Plugins/Media/MediaViewer/Source/MediaViewer/Private/Widgets/SMediaViewerDropTarget.cpp
--- a/Engine/Plugins/Media/MediaViewer/Source/MediaViewer/Private/Widgets/SMediaViewerDropTarget.cpp
+++ b/Engine/Plugins/Media/MediaViewer/Source/MediaViewer/Private/Widgets/SMediaViewerDropTarget.cpp
@@ -2,17 +2,25 @@
 
 #include "Widgets/SMediaViewerDropTarget.h"
 
+#include "AssetRegistry/AssetRegistryModule.h"
 #include "AssetSelection.h"
+#include "Containers/Array.h"
+#include "Containers/ArrayView.h"
+#include "ContentBrowserModule.h"
 #include "DragAndDrop/DecoratedDragDropOp.h"
 #include "FileMediaSource.h"
 #include "Framework/Application/SlateApplication.h"
+#include "IContentBrowserSingleton.h"
 #include "ImageViewers/MediaSourceImageViewer.h"
 #include "IMediaViewerModule.h"
+#include "ImgMediaSource.h"
 #include "Input/DragAndDrop.h"
 #include "Library/MediaViewerLibrary.h"
 #include "Library/MediaViewerLibraryItem.h"
 #include "MediaPlayer.h"
 #include "MediaViewerDelegates.h"
+#include "MediaViewerModule.h"
+#include "Misc/MessageDialog.h"
 #include "SDropTarget.h"
 #include "Styling/AppStyle.h"
 #include "Styling/StyleColors.h"
@@ -205,7 +213,7 @@ void SMediaViewerDropTarget::HandleDroppedAssets(const TArrayView<FAssetData> In
 			}
 		}
 	}
-
+	
 	if (ImageViewers.Num() > 1)
 	{
 		for (int32 Index = 0; Index < static_cast<int32>(EMediaImageViewerPosition::COUNT); ++Index)
@@ -229,6 +237,11 @@ void SMediaViewerDropTarget::HandleDroppedAssets(const TArrayView<FAssetData> In
 
 void SMediaViewerDropTarget::HandleDroppeFileOp(const FExternalDragOperation& InFileDragDropOp)
 {
+	if (CreateAndSetImageSequence(InFileDragDropOp))
+	{
+		return;
+	}
+
 	UMediaPlayer* TestMediaPlayer = NewObject<UMediaPlayer>(GetTransientPackage());
 	TArray<UFileMediaSource*, TInlineAllocator<2>> Sources;
 
@@ -279,6 +292,134 @@ void SMediaViewerDropTarget::HandleDroppeFileOp(const FExternalDragOperation& In
 	}
 }
 
+bool SMediaViewerDropTarget::CreateAndSetImageSequence(const FExternalDragOperation& InFileDragDropOp)
+{
+	FString ImageSequencePath;
+	const TArray<FString>& Files = InFileDragDropOp.GetFiles();
+
+	if (Files.IsEmpty())
+	{
+		return false;
+	}
+
+	// Drag and drop a directory - open as image sequence
+	if (Files.Num() == 1 && FPaths::DirectoryExists(Files[0]))
+	{
+		ImageSequencePath = Files[0];
+	}
+	// Drag and drop more than 2 files (2 would be a comparison view)
+	else if (Files.Num() > 2)
+	{
+		ImageSequencePath = FPaths::GetPath(Files[0]);
+
+		const TArray<FString> ValidTypes = { TEXT("BMP"), TEXT("EXR"), TEXT("PNG"), TEXT("JPG"), TEXT("JPEG") };
+
+		// Sanity check
+		if (!FPaths::DirectoryExists(ImageSequencePath))
+		{
+			return false;
+		}
+
+		if (!ValidTypes.Contains(FPaths::GetExtension(Files[0]).ToUpper()))
+		{
+			return false;
+		}
+
+		for (int32 Index = 1; Index < Files.Num(); ++Index)
+		{
+			// Check if it's in the same folder.
+			if (!FPaths::GetPath(Files[Index]).Equals(ImageSequencePath))
+			{
+				return false;
+			}
+
+			if (!ValidTypes.Contains(FPaths::GetExtension(Files[Index]).ToUpper()))
+			{
+				return false;
+			}
+		}
+	}
+	else
+	{
+		return false;
+	}
+
+	const EAppReturnType::Type Return = FMessageDialog::Open(
+		EAppMsgType::OkCancel,
+		LOCTEXT("ImportAsImageSequence", "Import as Image Media Source?")
+	);
+
+	if (Return != EAppReturnType::Ok)
+	{
+		return false;
+	}
+
+	FSaveAssetDialogConfig SaveAssetDialogConfig;
+	SaveAssetDialogConfig.DialogTitleOverride = LOCTEXT("CreateAssetDialogTitle", "Save As Image Media Source");
+	SaveAssetDialogConfig.DefaultPath = "/Game";
+	SaveAssetDialogConfig.DefaultAssetName = "ImageMediaSource";
+	SaveAssetDialogConfig.ExistingAssetPolicy = ESaveAssetDialogExistingAssetPolicy::Disallow;
+
+	FContentBrowserModule& ContentBrowserModule = FModuleManager::LoadModuleChecked<FContentBrowserModule>("ContentBrowser");
+	const FString SaveObjectPath = ContentBrowserModule.Get().CreateModalSaveAssetDialog(SaveAssetDialogConfig);
+
+	if (SaveObjectPath.IsEmpty())
+	{
+		return false;
+	}
+
+	const FString PackagePath = FPaths::GetBaseFilename(*SaveObjectPath, false);
+
+	UPackage* Package = CreatePackage(*PackagePath);
+
+	if (!Package)
+	{
+		return false;
+	}
+
+	const FString AssetName = FPaths::GetBaseFilename(*SaveObjectPath, true);
+
+	UImgMediaSource* ImageSequence = NewObject<UImgMediaSource>(Package, *AssetName, RF_Public | RF_Standalone | RF_Transactional);
+
+	if (!ImageSequence)
+	{
+		return false;
+	}
+
+	ImageSequence->SetSequencePath(ImageSequencePath);
+
+	FAssetRegistryModule::AssetCreated(ImageSequence);
+
+	IMediaViewerModule& MediaViewerModule = IMediaViewerModule::Get();
+	TSharedPtr<FMediaViewerLibraryItem> LibraryItem = MediaViewerModule.CreateLibraryItem(TNotNull<UObject*>(ImageSequence));
+	
+	if (!LibraryItem.IsValid())
+	{
+		return false;
+	}
+
+	TSharedPtr<FMediaImageViewer> ImageViewer = LibraryItem->CreateImageViewer();
+
+	if (!ImageViewer.IsValid())
+	{
+		return false;
+	}
+
+	Package->MarkPackageDirty();
+
+	// Swap images and then replace first image.
+	if (bForceComparisonView && Position == EMediaImageViewerPosition::First)
+	{
+		Delegates->SwapAB.Execute();
+		Delegates->SetABView.Execute();
+		Delegates->SetABOrientation.Execute(Orient_Horizontal);
+	}
+
+	Delegates->SetImageViewer.Execute(Position, ImageViewer.ToSharedRef());
+	
+	return true;
+}
+
 bool SMediaViewerDropTarget::OnAllowDrop(TSharedPtr<FDragDropOperation> InDragDropOperation) const
 {
 	if (!InDragDropOperation.IsValid())
diff --git a/Engine/Plugins/Media/MediaViewer/Source/MediaViewer/Private/Widgets/SMediaViewerDropTarget.h b/Engine/Plugins/Media/MediaViewer/Source/MediaViewer/Private/Widgets/SMediaViewerDropTarget.h
--- a/Engine/Plugins/Media/MediaViewer/Source/MediaViewer/Private/Widgets/SMediaViewerDropTarget.h
+++ b/Engine/Plugins/Media/MediaViewer/Source/MediaViewer/Private/Widgets/SMediaViewerDropTarget.h
@@ -39,6 +39,7 @@ public:
 	//~ End SWidget
 
 protected:
+	/** Filters the given assets and returns the ones with a registered image viewer handler class. */
 	static TArray<FAssetData> GetAssetsWithImageViewer(TConstArrayView<FAssetData> InAssets);
 
 	TSharedPtr<FMediaViewerDelegates> Delegates;
@@ -54,11 +55,34 @@ protected:
 
 	FReply OnDropped(const FGeometry& InGeometry, const FDragDropEvent& InDragDropEvent);
 
+	/** 
+	 * Takes the dragged library item(s) and displays them in the viewer. 
+	 * If one is dragged, it will replace the current image viewer.
+	 * If 2 or more are dragged, the first 2 will be put in comparison mode, replacing whatever is being shown.
+	 */
 	void HandleDroppedMediaViewerOp(const FMediaViewerLibraryItemDragDropOperation& InMediaViewerOp);
 
+	/**
+	 * Takes the dragged assets and displays them in the viewer, if they have a registered image viewer handler class.
+	 * If one is dragged, it will replace the current image viewer.
+	 * If 2 or more are dragged, the first 2 will be put in comparison mode, replacing whatever is being shown.
+	 */
 	void HandleDroppedAssets(const TArrayView<FAssetData> InDroppedAssets);
 	
+	/**
+	 * Takes the dragged files and either attempts to open them as a media source (video/images/etc.)
+	 * Does not support solo images, only a directory or group of images for an image media source.
+	 * If one is dragged, it will replace the current image viewer.
+	 * If 2 or more are dragged, the first 2 will be put in comparison mode, replacing whatever is being shown.
+	 */
 	void HandleDroppeFileOp(const FExternalDragOperation& InFileDragDropOp);
+
+	/**
+	 * Looks at the dragged files and determines if it can be turned into an image media source.
+	 * If it can, it will ask the user to save the asset and then open it in the media viewer.
+	 * Returns true if an asset was created.
+	 */
+	bool CreateAndSetImageSequence(const FExternalDragOperation& InFileDragDropOp);
 };
 
 } // UE::MediaViewer::Private
