Media Viewer: Fixed crash when open the editor viewport library section.

diff --git a/Engine/Plugins/Media/MediaViewer/Source/MediaViewer/Private/ImageViewers/TextureSampleCache.cpp b/Engine/Plugins/Media/MediaViewer/Source/MediaViewer/Private/ImageViewers/TextureSampleCache.cpp
--- a/Engine/Plugins/Media/MediaViewer/Source/MediaViewer/Private/ImageViewers/TextureSampleCache.cpp
+++ b/Engine/Plugins/Media/MediaViewer/Source/MediaViewer/Private/ImageViewers/TextureSampleCache.cpp
@@ -201,6 +201,11 @@ void FTextureSampleCache::SetPixelColor_RenderTarget(const FIntPoint& InPixelCoo
 			),
 			/* Transparent */ true
 		);
+
+		if (!RenderTarget)
+		{
+			return;
+		}
 	}
 
 	bool bNeedsUpdate = false;
@@ -232,7 +237,7 @@ void FTextureSampleCache::SetPixelColor_RenderTarget(const FIntPoint& InPixelCoo
 		{
 			FTextureResource* SourceResource = Source->GetResource();
 
-			if (!SourceResource)
+			if (!SourceResource || !SourceResource->TextureRHI.IsValid())
 			{
 				OnComplete.ExecuteIfBound(false);
 				return;
@@ -248,7 +253,7 @@ void FTextureSampleCache::SetPixelColor_RenderTarget(const FIntPoint& InPixelCoo
 
 			FTextureResource* TargetResource = Target->GetResource();
 
-			if (!TargetResource)
+			if (!TargetResource || !TargetResource->TextureRHI.IsValid())
 			{
 				OnComplete.ExecuteIfBound(false);
 				return;
