Media Viewer: Added mirroring and gamma fixes.
- Images can now be flipped horiontally and vertically.
- Correct gamma correction is now applied to sampled colors.
- A color chip is now in status bar with the color information.

#Synced-CL 49593614

diff --git a/Engine/Plugins/Media/MediaViewer/Source/MediaViewer/Private/ImageViewers/MaterialInterfaceImageViewer.cpp b/Engine/Plugins/Media/MediaViewer/Source/MediaViewer/Private/ImageViewers/MaterialInterfaceImageViewer.cpp
--- a/Engine/Plugins/Media/MediaViewer/Source/MediaViewer/Private/ImageViewers/MaterialInterfaceImageViewer.cpp
+++ b/Engine/Plugins/Media/MediaViewer/Source/MediaViewer/Private/ImageViewers/MaterialInterfaceImageViewer.cpp
@@ -314,10 +314,22 @@ TOptional<TVariant<FColor, FLinearColor>> FMaterialInterfaceImageViewer::GetPixe
 		return {};
 	}
 
-	if (const FLinearColor* PixelColor = SampleCache->GetPixelColor(InPixelCoords))
+	FIntPoint MirroredPixelCoords = InPixelCoords;
+
+	if (EnumHasAllFlags(static_cast<EMediaImageMirrorPlane>(PaintSettings.MirrorPlanes), EMediaImageMirrorPlane::Horizontal))
+	{
+		MirroredPixelCoords.X = ImageInfo.Size.X - InPixelCoords.X - 1;
+	}
+
+	if (EnumHasAllFlags(static_cast<EMediaImageMirrorPlane>(PaintSettings.MirrorPlanes), EMediaImageMirrorPlane::Vertical))
+	{
+		MirroredPixelCoords.Y = ImageInfo.Size.Y - InPixelCoords.Y - 1;
+	}
+
+	if (const FLinearColor* PixelColor = SampleCache->GetPixelColor(MirroredPixelCoords))
 	{
 		TVariant<FColor, FLinearColor> PixelColorVariant;
-		PixelColorVariant.Set<FLinearColor>(*PixelColor);
+		PixelColorVariant.Set<FLinearColor>(FMediaViewerUtils::CorrectGamma(*PixelColor));
 
 		return PixelColorVariant;
 	}
@@ -347,10 +359,24 @@ void FMaterialInterfaceImageViewer::NotifyPostChange(const FPropertyChangedEvent
 {
 	FMediaImageViewer::NotifyPostChange(InPropertyChangedEvent, InPropertyThatChanged);
 
-	if (InPropertyChangedEvent.GetMemberPropertyName() == GET_MEMBER_NAME_CHECKED(FMaterialInterfaceImageViewerSettings, RenderTargetSize))
+	const FName MemberName = InPropertyChangedEvent.GetMemberPropertyName();
+
+	if (MemberName == GET_MEMBER_NAME_CHECKED(FMaterialInterfaceImageViewerSettings, RenderTargetSize))
 	{
 		CreateBrush();
 	}
+	else if (MemberName == GET_MEMBER_NAME_CHECKED(FMediaImagePaintSettings, MirrorPlanes))
+	{
+		OnMirrored();
+	}
+}
+
+void FMaterialInterfaceImageViewer::OnMirrored()
+{
+	if (SampleCache.IsValid())
+	{
+		SampleCache->Invalidate();
+	}
 }
 
 FString FMaterialInterfaceImageViewer::GetReferencerName() const
diff --git a/Engine/Plugins/Media/MediaViewer/Source/MediaViewer/Private/ImageViewers/MaterialInterfaceImageViewer.h b/Engine/Plugins/Media/MediaViewer/Source/MediaViewer/Private/ImageViewers/MaterialInterfaceImageViewer.h
--- a/Engine/Plugins/Media/MediaViewer/Source/MediaViewer/Private/ImageViewers/MaterialInterfaceImageViewer.h
+++ b/Engine/Plugins/Media/MediaViewer/Source/MediaViewer/Private/ImageViewers/MaterialInterfaceImageViewer.h
@@ -95,6 +95,7 @@ public:
 	virtual bool OnOwnerCleanup(UObject* InObject) override;
 	virtual TOptional<TVariant<FColor, FLinearColor>> GetPixelColor(const FIntPoint& InPixelCoords, int32 InMipLevel) const override;
 	virtual TSharedPtr<FStructOnScope> GetCustomSettingsOnScope() const override;
+	virtual void OnMirrored() override;
 	virtual void PaintImage(FMediaImageSlatePaintParams& InPaintParams, const FMediaImageSlatePaintGeometry& InPaintGeometry) override;
 	//~ End FMediaImageViewer
 
diff --git a/Engine/Plugins/Media/MediaViewer/Source/MediaViewer/Private/ImageViewers/MediaImageViewer.cpp b/Engine/Plugins/Media/MediaViewer/Source/MediaViewer/Private/ImageViewers/MediaImageViewer.cpp
--- a/Engine/Plugins/Media/MediaViewer/Source/MediaViewer/Private/ImageViewers/MediaImageViewer.cpp
+++ b/Engine/Plugins/Media/MediaViewer/Source/MediaViewer/Private/ImageViewers/MediaImageViewer.cpp
@@ -85,6 +85,30 @@ void FMediaImageViewer::AddReferencedObjects(FReferenceCollector& InCollector)
 	}
 }
 
+bool FMediaImageViewer::IsMirrored(EMediaImageMirrorPlane InPlane) const
+{
+	return EnumHasAllFlags(static_cast<EMediaImageMirrorPlane>(PaintSettings.MirrorPlanes), InPlane);
+}
+
+void FMediaImageViewer::SetMirrored(EMediaImageMirrorPlane InPlane, bool bInSet)
+{
+	if (IsMirrored(InPlane) == bInSet)
+	{
+		return;
+	}
+
+	if (bInSet)
+	{
+		PaintSettings.MirrorPlanes = PaintSettings.MirrorPlanes | static_cast<uint8>(InPlane);
+	}
+	else
+	{
+		PaintSettings.MirrorPlanes = PaintSettings.MirrorPlanes & !static_cast<uint8>(InPlane);
+	}
+
+	OnMirrored();
+}
+
 FVector2D FMediaImageViewer::GetViewerCenter(const FVector2D& InViewerSize) const
 {
 	return InViewerSize * 0.5;
@@ -298,6 +322,25 @@ void FMediaImageViewer::PaintImage(FMediaImageSlatePaintParams& InPaintParams, c
 		return;
 	}
 
+	switch (PaintSettings.MirrorPlanes & static_cast<uint8>(EMediaImageMirrorPlane::TwoDim))
+	{
+		case static_cast<uint8>(EMediaImageMirrorPlane::None):
+			BrushToUse->Mirroring = ESlateBrushMirrorType::NoMirror;
+			break;
+
+		case static_cast<uint8>(EMediaImageMirrorPlane::Horizontal):
+			BrushToUse->Mirroring = ESlateBrushMirrorType::Horizontal;
+			break;
+
+		case static_cast<uint8>(EMediaImageMirrorPlane::Vertical):
+			BrushToUse->Mirroring = ESlateBrushMirrorType::Vertical;
+			break;
+
+		case static_cast<uint8>(EMediaImageMirrorPlane::Horizontal | EMediaImageMirrorPlane::Vertical):
+			BrushToUse->Mirroring = ESlateBrushMirrorType::Both;
+			break;
+	}
+
 	FSlateDrawElement::MakeRotatedBox(
 		InPaintParams.DrawElements,
 		InPaintParams.LayerId,
diff --git a/Engine/Plugins/Media/MediaViewer/Source/MediaViewer/Private/ImageViewers/MediaSourceImageViewer.cpp b/Engine/Plugins/Media/MediaViewer/Source/MediaViewer/Private/ImageViewers/MediaSourceImageViewer.cpp
--- a/Engine/Plugins/Media/MediaViewer/Source/MediaViewer/Private/ImageViewers/MediaSourceImageViewer.cpp
+++ b/Engine/Plugins/Media/MediaViewer/Source/MediaViewer/Private/ImageViewers/MediaSourceImageViewer.cpp
@@ -398,6 +398,18 @@ TOptional<TVariant<FColor, FLinearColor>> FMediaSourceImageViewer::GetPixelColor
 		return {};
 	}
 
+	FIntPoint MirroredPixelCoords = InPixelCoords;
+
+	if (EnumHasAllFlags(static_cast<EMediaImageMirrorPlane>(PaintSettings.MirrorPlanes), EMediaImageMirrorPlane::Horizontal))
+	{
+		MirroredPixelCoords.X = ImageInfo.Size.X - InPixelCoords.X - 1;
+	}
+
+	if (EnumHasAllFlags(static_cast<EMediaImageMirrorPlane>(PaintSettings.MirrorPlanes), EMediaImageMirrorPlane::Vertical))
+	{
+		MirroredPixelCoords.Y = ImageInfo.Size.Y - InPixelCoords.Y - 1;
+	}
+
 	if (!MediaStream)
 	{
 		SampleCache->Invalidate();
@@ -416,7 +428,7 @@ TOptional<TVariant<FColor, FLinearColor>> FMediaSourceImageViewer::GetPixelColor
 	{
 		const FTimespan PlayerTime = Player->GetTime();
 
-		if (const FLinearColor* PixelColor = SampleCache->GetPixelColor(InPixelCoords, PlayerTime))
+		if (const FLinearColor* PixelColor = SampleCache->GetPixelColor(MirroredPixelCoords, PlayerTime))
 		{
 			TVariant<FColor, FLinearColor> PixelColorVariant;
 			PixelColorVariant.Set<FLinearColor>(*PixelColor);
@@ -451,6 +463,24 @@ void FMediaSourceImageViewer::AddReferencedObjects(FReferenceCollector& InCollec
 	}
 }
 
+void FMediaSourceImageViewer::NotifyPostChange(const FPropertyChangedEvent& InPropertyChangedEvent, FProperty* InPropertyThatChanged)
+{
+	FMediaImageViewer::NotifyPostChange(InPropertyChangedEvent, InPropertyThatChanged);
+
+	if (InPropertyChangedEvent.GetMemberPropertyName() == GET_MEMBER_NAME_CHECKED(FMediaImagePaintSettings, MirrorPlanes))
+	{
+		OnMirrored();
+	}
+}
+
+void FMediaSourceImageViewer::OnMirrored()
+{
+	if (SampleCache.IsValid())
+	{
+		SampleCache->Invalidate();
+	}
+}
+
 TSharedPtr<SWidget> FMediaSourceImageViewer::GetOverlayWidget(EMediaImageViewerPosition InPosition, const TSharedPtr<FMediaViewerDelegates>& InDelegates)
 {
 	return SNew(SBox)
diff --git a/Engine/Plugins/Media/MediaViewer/Source/MediaViewer/Private/ImageViewers/MediaSourceImageViewer.h b/Engine/Plugins/Media/MediaViewer/Source/MediaViewer/Private/ImageViewers/MediaSourceImageViewer.h
--- a/Engine/Plugins/Media/MediaViewer/Source/MediaViewer/Private/ImageViewers/MediaSourceImageViewer.h
+++ b/Engine/Plugins/Media/MediaViewer/Source/MediaViewer/Private/ImageViewers/MediaSourceImageViewer.h
@@ -127,6 +127,7 @@ public:
 	virtual bool OnOwnerCleanup(UObject* InObject) override;
 	virtual TOptional<TVariant<FColor, FLinearColor>> GetPixelColor(const FIntPoint& InPixelCoords, int32 InMipLevel) const override;
 	virtual TSharedPtr<FStructOnScope> GetCustomSettingsOnScope() const override;
+	virtual void OnMirrored() override;
 	virtual TSharedPtr<SWidget> GetOverlayWidget(EMediaImageViewerPosition InPosition, const TSharedPtr<FMediaViewerDelegates>& InDelegates) override;
 	virtual void ExtendStatusBar(UE::MediaViewer::FMediaImageStatusBarExtender& InOutStatusBarExtender) override;
 	//~ End FMediaImageViewer
@@ -136,6 +137,10 @@ public:
 	virtual void AddReferencedObjects(FReferenceCollector& InCollector) override;
 	//~ End FGCObject
 
+	//~ Begin FNotifyHook
+	virtual void NotifyPostChange(const FPropertyChangedEvent& InPropertyChangedEvent, FProperty* InPropertyThatChanged) override;
+	//~ End FNotifyHook
+
 protected:
 	FMediaSourceImageViewerSettings MediaSourceSettings;
 	TObjectPtr<UMediaStream> MediaStream;
diff --git a/Engine/Plugins/Media/MediaViewer/Source/MediaViewer/Private/ImageViewers/MediaTextureImageViewer.cpp b/Engine/Plugins/Media/MediaViewer/Source/MediaViewer/Private/ImageViewers/MediaTextureImageViewer.cpp
--- a/Engine/Plugins/Media/MediaViewer/Source/MediaViewer/Private/ImageViewers/MediaTextureImageViewer.cpp
+++ b/Engine/Plugins/Media/MediaViewer/Source/MediaViewer/Private/ImageViewers/MediaTextureImageViewer.cpp
@@ -253,10 +253,22 @@ TOptional<TVariant<FColor, FLinearColor>> FMediaTextureImageViewer::GetPixelColo
 		return {};
 	}
 
-	if (const FLinearColor* PixelColor = SampleCache->GetPixelColor(InPixelCoords))
+	FIntPoint MirroredPixelCoords = InPixelCoords;
+
+	if (EnumHasAllFlags(static_cast<EMediaImageMirrorPlane>(PaintSettings.MirrorPlanes), EMediaImageMirrorPlane::Horizontal))
+	{
+		MirroredPixelCoords.X = ImageInfo.Size.X - InPixelCoords.X - 1;
+	}
+
+	if (EnumHasAllFlags(static_cast<EMediaImageMirrorPlane>(PaintSettings.MirrorPlanes), EMediaImageMirrorPlane::Vertical))
+	{
+		MirroredPixelCoords.Y = ImageInfo.Size.Y - InPixelCoords.Y - 1;
+	}
+
+	if (const FLinearColor* PixelColor = SampleCache->GetPixelColor(MirroredPixelCoords))
 	{
 		TVariant<FColor, FLinearColor> PixelColorVariant;
-		PixelColorVariant.Set<FLinearColor>(*PixelColor);
+		PixelColorVariant.Set<FLinearColor>(FMediaViewerUtils::CorrectGamma(*PixelColor));
 
 		return PixelColorVariant;
 	}
@@ -282,6 +294,24 @@ void FMediaTextureImageViewer::AddReferencedObjects(FReferenceCollector& InColle
 	);
 }
 
+void FMediaTextureImageViewer::NotifyPostChange(const FPropertyChangedEvent& InPropertyChangedEvent, FProperty* InPropertyThatChanged)
+{
+	FMediaImageViewer::NotifyPostChange(InPropertyChangedEvent, InPropertyThatChanged);
+
+	if (InPropertyChangedEvent.GetMemberPropertyName() == GET_MEMBER_NAME_CHECKED(FMediaImagePaintSettings, MirrorPlanes))
+	{
+		OnMirrored();
+	}
+}
+
+void FMediaTextureImageViewer::OnMirrored()
+{
+	if (SampleCache.IsValid())
+	{
+		SampleCache->Invalidate();
+	}
+}
+
 bool FMediaTextureImageViewer::OnOwnerCleanup(UObject* InObject)
 {
 	if (InObject && MediaTextureSettings.MediaTexture && 
diff --git a/Engine/Plugins/Media/MediaViewer/Source/MediaViewer/Private/ImageViewers/MediaTextureImageViewer.h b/Engine/Plugins/Media/MediaViewer/Source/MediaViewer/Private/ImageViewers/MediaTextureImageViewer.h
--- a/Engine/Plugins/Media/MediaViewer/Source/MediaViewer/Private/ImageViewers/MediaTextureImageViewer.h
+++ b/Engine/Plugins/Media/MediaViewer/Source/MediaViewer/Private/ImageViewers/MediaTextureImageViewer.h
@@ -83,6 +83,7 @@ public:
 	virtual bool OnOwnerCleanup(UObject* InObject) override;
 	virtual TOptional<TVariant<FColor, FLinearColor>> GetPixelColor(const FIntPoint& InPixelCoords, int32 InMipLevel) const override;
 	virtual TSharedPtr<FStructOnScope> GetCustomSettingsOnScope() const override;
+	virtual void OnMirrored() override;
 	//~ End FMediaImageViewer
 
 	//~ Begin FGCObject
@@ -90,6 +91,10 @@ public:
 	virtual void AddReferencedObjects(FReferenceCollector& InCollector) override;
 	//~ End FGCObject
 
+	//~ Begin FNotifyHook
+	virtual void NotifyPostChange(const FPropertyChangedEvent& InPropertyChangedEvent, FProperty* InPropertyThatChanged) override;
+	//~ End FNotifyHook
+
 protected:
 	FMediaTextureImageViewerSettings MediaTextureSettings;
 	TSharedPtr<FTextureSampleCache> SampleCache;
diff --git a/Engine/Plugins/Media/MediaViewer/Source/MediaViewer/Private/ImageViewers/Texture2DImageViewer.cpp b/Engine/Plugins/Media/MediaViewer/Source/MediaViewer/Private/ImageViewers/Texture2DImageViewer.cpp
--- a/Engine/Plugins/Media/MediaViewer/Source/MediaViewer/Private/ImageViewers/Texture2DImageViewer.cpp
+++ b/Engine/Plugins/Media/MediaViewer/Source/MediaViewer/Private/ImageViewers/Texture2DImageViewer.cpp
@@ -277,10 +277,22 @@ TOptional<TVariant<FColor, FLinearColor>> FTexture2DImageViewer::GetPixelColor(c
 		return {};
 	}
 
-	if (const FLinearColor* PixelColor = SampleCache->GetPixelColor(InPixelCoords))
+	FIntPoint MirroredPixelCoords = InPixelCoords;
+
+	if (EnumHasAllFlags(static_cast<EMediaImageMirrorPlane>(PaintSettings.MirrorPlanes), EMediaImageMirrorPlane::Horizontal))
+	{
+		MirroredPixelCoords.X = ImageInfo.Size.X - InPixelCoords.X - 1;
+	}
+
+	if (EnumHasAllFlags(static_cast<EMediaImageMirrorPlane>(PaintSettings.MirrorPlanes), EMediaImageMirrorPlane::Vertical))
+	{
+		MirroredPixelCoords.Y = ImageInfo.Size.Y - InPixelCoords.Y - 1;
+	}
+
+	if (const FLinearColor* PixelColor = SampleCache->GetPixelColor(MirroredPixelCoords))
 	{
 		TVariant<FColor, FLinearColor> PixelColorVariant;
-		PixelColorVariant.Set<FLinearColor>(*PixelColor);
+		PixelColorVariant.Set<FLinearColor>(FMediaViewerUtils::CorrectGamma(*PixelColor));
 
 		return PixelColorVariant;
 	}
@@ -306,6 +318,24 @@ void FTexture2DImageViewer::AddReferencedObjects(FReferenceCollector& InCollecto
 	);
 }
 
+void FTexture2DImageViewer::NotifyPostChange(const FPropertyChangedEvent& InPropertyChangedEvent, FProperty* InPropertyThatChanged)
+{
+	FMediaImageViewer::NotifyPostChange(InPropertyChangedEvent, InPropertyThatChanged);
+
+	if (InPropertyChangedEvent.GetMemberPropertyName() == GET_MEMBER_NAME_CHECKED(FMediaImagePaintSettings, MirrorPlanes))
+	{
+		OnMirrored();
+	}
+}
+
+void FTexture2DImageViewer::OnMirrored()
+{
+	if (SampleCache.IsValid())
+	{
+		SampleCache->Invalidate();
+	}
+}
+
 TOptional<FIntPoint> FTexture2DImageViewer::GetTextureSize() const
 {
 	if (!TextureSettings.Texture)
@@ -357,6 +387,7 @@ void FTexture2DImageViewer::PaintImage(FMediaImageSlatePaintParams& InPaintParam
 		{
 			ImageInfo.Size.X = TextureSettings.Texture->GetSurfaceWidth();
 			ImageInfo.Size.Y = TextureSettings.Texture->GetSurfaceHeight();
+			bValidImageSize = true;
 
 			if (Brush.IsValid())
 			{
diff --git a/Engine/Plugins/Media/MediaViewer/Source/MediaViewer/Private/ImageViewers/Texture2DImageViewer.h b/Engine/Plugins/Media/MediaViewer/Source/MediaViewer/Private/ImageViewers/Texture2DImageViewer.h
--- a/Engine/Plugins/Media/MediaViewer/Source/MediaViewer/Private/ImageViewers/Texture2DImageViewer.h
+++ b/Engine/Plugins/Media/MediaViewer/Source/MediaViewer/Private/ImageViewers/Texture2DImageViewer.h
@@ -89,6 +89,7 @@ public:
 	virtual bool OnOwnerCleanup(UObject* InObject) override;
 	virtual TOptional<TVariant<FColor, FLinearColor>> GetPixelColor(const FIntPoint& InPixelCoords, int32 InMipLevel) const override;
 	virtual TSharedPtr<FStructOnScope> GetCustomSettingsOnScope() const override;
+	virtual void OnMirrored() override;
 	//~ End FMediaImageViewer
 
 	//~ Begin FGCObject
@@ -96,6 +97,10 @@ public:
 	virtual void AddReferencedObjects(FReferenceCollector& InCollector) override;
 	//~ End FGCObject
 
+	//~ Begin FNotifyHook
+	virtual void NotifyPostChange(const FPropertyChangedEvent& InPropertyChangedEvent, FProperty* InPropertyThatChanged) override;
+	//~ End FNotifyHook
+
 protected:
 	FTexture2DImageViewerSettings TextureSettings;
 	TSharedPtr<FTextureSampleCache> SampleCache;
diff --git a/Engine/Plugins/Media/MediaViewer/Source/MediaViewer/Private/ImageViewers/TextureRenderTarget2DImageViewer.cpp b/Engine/Plugins/Media/MediaViewer/Source/MediaViewer/Private/ImageViewers/TextureRenderTarget2DImageViewer.cpp
--- a/Engine/Plugins/Media/MediaViewer/Source/MediaViewer/Private/ImageViewers/TextureRenderTarget2DImageViewer.cpp
+++ b/Engine/Plugins/Media/MediaViewer/Source/MediaViewer/Private/ImageViewers/TextureRenderTarget2DImageViewer.cpp
@@ -247,10 +247,22 @@ TOptional<TVariant<FColor, FLinearColor>> FTextureRenderTarget2DImageViewer::Get
 		return {};
 	}
 
-	if (const FLinearColor* PixelColor = SampleCache->GetPixelColor(InPixelCoords))
+	FIntPoint MirroredPixelCoords = InPixelCoords;
+
+	if (EnumHasAllFlags(static_cast<EMediaImageMirrorPlane>(PaintSettings.MirrorPlanes), EMediaImageMirrorPlane::Horizontal))
+	{
+		MirroredPixelCoords.X = ImageInfo.Size.X - InPixelCoords.X - 1;
+	}
+
+	if (EnumHasAllFlags(static_cast<EMediaImageMirrorPlane>(PaintSettings.MirrorPlanes), EMediaImageMirrorPlane::Vertical))
+	{
+		MirroredPixelCoords.Y = ImageInfo.Size.Y - InPixelCoords.Y - 1;
+	}
+
+	if (const FLinearColor* PixelColor = SampleCache->GetPixelColor(MirroredPixelCoords))
 	{
 		TVariant<FColor, FLinearColor> PixelColorVariant;
-		PixelColorVariant.Set<FLinearColor>(*PixelColor);
+		PixelColorVariant.Set<FLinearColor>(FMediaViewerUtils::CorrectGamma(*PixelColor));
 
 		return PixelColorVariant;
 	}
@@ -276,6 +288,24 @@ void FTextureRenderTarget2DImageViewer::AddReferencedObjects(FReferenceCollector
 	);
 }
 
+void FTextureRenderTarget2DImageViewer::NotifyPostChange(const FPropertyChangedEvent& InPropertyChangedEvent, FProperty* InPropertyThatChanged)
+{
+	FMediaImageViewer::NotifyPostChange(InPropertyChangedEvent, InPropertyThatChanged);
+
+	if (InPropertyChangedEvent.GetMemberPropertyName() == GET_MEMBER_NAME_CHECKED(FMediaImagePaintSettings, MirrorPlanes))
+	{
+		OnMirrored();
+	}
+}
+
+void FTextureRenderTarget2DImageViewer::OnMirrored()
+{
+	if (SampleCache.IsValid())
+	{
+		SampleCache->Invalidate();
+	}
+}
+
 bool FTextureRenderTarget2DImageViewer::OnOwnerCleanup(UObject* InObject)
 {
 	if (InObject && RenderTargetSettings.RenderTarget && 
diff --git a/Engine/Plugins/Media/MediaViewer/Source/MediaViewer/Private/ImageViewers/TextureRenderTarget2DImageViewer.h b/Engine/Plugins/Media/MediaViewer/Source/MediaViewer/Private/ImageViewers/TextureRenderTarget2DImageViewer.h
--- a/Engine/Plugins/Media/MediaViewer/Source/MediaViewer/Private/ImageViewers/TextureRenderTarget2DImageViewer.h
+++ b/Engine/Plugins/Media/MediaViewer/Source/MediaViewer/Private/ImageViewers/TextureRenderTarget2DImageViewer.h
@@ -80,6 +80,7 @@ public:
 	virtual bool OnOwnerCleanup(UObject* InObject) override;
 	virtual TOptional<TVariant<FColor, FLinearColor>> GetPixelColor(const FIntPoint& InPixelCoords, int32 InMipLevel) const override;
 	virtual TSharedPtr<FStructOnScope> GetCustomSettingsOnScope() const override;
+	virtual void OnMirrored() override;
 	//~ End FMediaImageViewer
 
 	//~ Begin FGCObject
@@ -87,6 +88,10 @@ public:
 	virtual void AddReferencedObjects(FReferenceCollector& InCollector) override;
 	//~ End FGCObject
 
+	//~ Begin FNotifyHook
+	virtual void NotifyPostChange(const FPropertyChangedEvent& InPropertyChangedEvent, FProperty* InPropertyThatChanged) override;
+	//~ End FNotifyHook
+
 protected:
 	FTextureRenderTarget2DImageViewerSettings RenderTargetSettings;
 	TSharedPtr<FTextureSampleCache> SampleCache;
diff --git a/Engine/Plugins/Media/MediaViewer/Source/MediaViewer/Private/ImageViewers/TextureSampleCache.cpp b/Engine/Plugins/Media/MediaViewer/Source/MediaViewer/Private/ImageViewers/TextureSampleCache.cpp
--- a/Engine/Plugins/Media/MediaViewer/Source/MediaViewer/Private/ImageViewers/TextureSampleCache.cpp
+++ b/Engine/Plugins/Media/MediaViewer/Source/MediaViewer/Private/ImageViewers/TextureSampleCache.cpp
@@ -12,6 +12,7 @@
 #include "RenderGraphUtils.h"
 #include "RenderingThread.h"
 #include "TextureResource.h"
+#include "UObject/Package.h"
 #include "UObject/StrongObjectPtr.h"
 #include "UObject/WeakObjectPtr.h"
 
@@ -194,12 +195,37 @@ void FTextureSampleCache::SetPixelColor_RenderTarget(const FIntPoint& InPixelCoo
 
 	if (!RenderTarget)
 	{
+		// We need to know if the texture is float or int.
+		bool bIsFloat = false;
+	
+		// todo: find a way to get the texture's pixel format. Now using rhi directly.
+		if (const FTextureResource* Resource = Texture->GetResource())
+		{
+			if (Resource->TextureRHI)
+			{
+				bIsFloat = IsFloatFormat(Resource->TextureRHI->GetDesc().Format);
+			}
+		}
+
+		ETextureRenderTargetFormat RenderTargetFormat;
+
+		if (bIsFloat)
+		{
+			RenderTargetFormat = RTF_RGBA16f;
+		}
+		else
+		{
+			const EGammaSpace Encoding = Texture->GetGammaSpace();
+			RenderTargetFormat = (Encoding == EGammaSpace::Linear) ? RTF_RGBA8 : RTF_RGBA8_SRGB;
+		}
+
 		RenderTarget = FMediaViewerUtils::CreateRenderTarget(
 			FIntPoint(
 				FMath::RoundToInt(Texture->GetSurfaceWidth()),
 				FMath::RoundToInt(Texture->GetSurfaceHeight())
 			),
-			/* Transparent */ true
+			/* Transparent */ true,
+			RenderTargetFormat
 		);
 
 		if (!RenderTarget)
diff --git a/Engine/Plugins/Media/MediaViewer/Source/MediaViewer/Private/ImageViewers/ViewportImageViewer.cpp b/Engine/Plugins/Media/MediaViewer/Source/MediaViewer/Private/ImageViewers/ViewportImageViewer.cpp
--- a/Engine/Plugins/Media/MediaViewer/Source/MediaViewer/Private/ImageViewers/ViewportImageViewer.cpp
+++ b/Engine/Plugins/Media/MediaViewer/Source/MediaViewer/Private/ImageViewers/ViewportImageViewer.cpp
@@ -57,7 +57,19 @@ TOptional<TVariant<FColor, FLinearColor>> FViewportImageViewer::GetPixelColor(co
 		return {};
 	}
 
-	if (const FLinearColor* PixelColor = SampleCache->GetPixelColor(InPixelCoords))
+	FIntPoint MirroredPixelCoords = InPixelCoords;
+
+	if (EnumHasAllFlags(static_cast<EMediaImageMirrorPlane>(PaintSettings.MirrorPlanes), EMediaImageMirrorPlane::Horizontal))
+	{
+		MirroredPixelCoords.X = ImageInfo.Size.X - InPixelCoords.X - 1;
+	}
+
+	if (EnumHasAllFlags(static_cast<EMediaImageMirrorPlane>(PaintSettings.MirrorPlanes), EMediaImageMirrorPlane::Vertical))
+	{
+		MirroredPixelCoords.Y = ImageInfo.Size.Y - InPixelCoords.Y - 1;
+	}
+
+	if (const FLinearColor* PixelColor = SampleCache->GetPixelColor(MirroredPixelCoords))
 	{
 		TVariant<FColor, FLinearColor> PixelColorVariant;
 		PixelColorVariant.Set<FLinearColor>(*PixelColor);
@@ -103,6 +115,24 @@ void FViewportImageViewer::AddReferencedObjects(FReferenceCollector& InCollector
 	);
 }
 
+void FViewportImageViewer::NotifyPostChange(const FPropertyChangedEvent& InPropertyChangedEvent, FProperty* InPropertyThatChanged)
+{
+	FMediaImageViewer::NotifyPostChange(InPropertyChangedEvent, InPropertyThatChanged);
+
+	if (InPropertyChangedEvent.GetMemberPropertyName() == GET_MEMBER_NAME_CHECKED(FMediaImagePaintSettings, MirrorPlanes))
+	{
+		OnMirrored();
+	}
+}
+
+void FViewportImageViewer::OnMirrored()
+{
+	if (SampleCache.IsValid())
+	{
+		SampleCache->Invalidate();
+	}
+}
+
 void FViewportImageViewer::CreateBrush()
 {
 	FViewport* Viewport = GetViewport();
diff --git a/Engine/Plugins/Media/MediaViewer/Source/MediaViewer/Private/ImageViewers/ViewportImageViewer.h b/Engine/Plugins/Media/MediaViewer/Source/MediaViewer/Private/ImageViewers/ViewportImageViewer.h
--- a/Engine/Plugins/Media/MediaViewer/Source/MediaViewer/Private/ImageViewers/ViewportImageViewer.h
+++ b/Engine/Plugins/Media/MediaViewer/Source/MediaViewer/Private/ImageViewers/ViewportImageViewer.h
@@ -43,6 +43,7 @@ public:
 	virtual bool IsValid() const override;
 	virtual TOptional<TVariant<FColor, FLinearColor>> GetPixelColor(const FIntPoint& InPixelCoords, int32 InMipLevel) const override;
 	virtual TSharedPtr<FStructOnScope> GetCustomSettingsOnScope() const override;
+	virtual void OnMirrored() override;
 	//~ End FMediaImageViewer
 
 	//~ Begin FGCObject
@@ -50,6 +51,10 @@ public:
 	virtual void AddReferencedObjects(FReferenceCollector& InCollector) override;
 	//~ End FGCObject
 
+	//~ Begin FNotifyHook
+	virtual void NotifyPostChange(const FPropertyChangedEvent& InPropertyChangedEvent, FProperty* InPropertyThatChanged) override;
+	//~ End FNotifyHook
+
 protected:
 	enum class ERenderState : uint8
 	{
diff --git a/Engine/Plugins/Media/MediaViewer/Source/MediaViewer/Private/MediaViewerCommands.cpp b/Engine/Plugins/Media/MediaViewer/Source/MediaViewer/Private/MediaViewerCommands.cpp
--- a/Engine/Plugins/Media/MediaViewer/Source/MediaViewer/Private/MediaViewerCommands.cpp
+++ b/Engine/Plugins/Media/MediaViewer/Source/MediaViewer/Private/MediaViewerCommands.cpp
@@ -266,6 +266,30 @@ void FMediaViewerCommands::RegisterCommands()
 		FInputChord(EKeys::T)
 	);
 
+	UI_COMMAND(
+		MirrorHorizonal,
+		"Mirror Horizontally",
+		"Mirror the image on the horizontal axis.",
+		EUserInterfaceActionType::ToggleButton,
+		FInputChord(EKeys::H)
+	);
+
+	UI_COMMAND(
+		MirrorVertical,
+		"Mirror Vertically",
+		"Mirror the image on the vertical axis.",
+		EUserInterfaceActionType::ToggleButton,
+		FInputChord(EKeys::V)
+	);
+
+	UI_COMMAND(
+		MirrorDepth,
+		"Mirror Depth",
+		"Mirror the image on the depth axis.",
+		EUserInterfaceActionType::ToggleButton,
+		FInputChord(EKeys::B)
+	);
+
 	UI_COMMAND(
 		SecondImageOpacityMinus,
 		"Lower B Image Opacity",
diff --git a/Engine/Plugins/Media/MediaViewer/Source/MediaViewer/Private/MediaViewerCommands.h b/Engine/Plugins/Media/MediaViewer/Source/MediaViewer/Private/MediaViewerCommands.h
--- a/Engine/Plugins/Media/MediaViewer/Source/MediaViewer/Private/MediaViewerCommands.h
+++ b/Engine/Plugins/Media/MediaViewer/Source/MediaViewer/Private/MediaViewerCommands.h
@@ -53,6 +53,10 @@ public:
 	TSharedPtr<FUICommandInfo> ResetAllTransforms;
 	TSharedPtr<FUICommandInfo> ToggleLockedTransform;
 
+	TSharedPtr<FUICommandInfo> MirrorHorizonal;
+	TSharedPtr<FUICommandInfo> MirrorVertical;
+	TSharedPtr<FUICommandInfo> MirrorDepth;
+
 	TSharedPtr<FUICommandInfo> SecondImageOpacityMinus;
 	TSharedPtr<FUICommandInfo> SecondImageOpacityPlus;
 
diff --git a/Engine/Plugins/Media/MediaViewer/Source/MediaViewer/Private/MediaViewerUtils.cpp b/Engine/Plugins/Media/MediaViewer/Source/MediaViewer/Private/MediaViewerUtils.cpp
--- a/Engine/Plugins/Media/MediaViewer/Source/MediaViewer/Private/MediaViewerUtils.cpp
+++ b/Engine/Plugins/Media/MediaViewer/Source/MediaViewer/Private/MediaViewerUtils.cpp
@@ -3,11 +3,10 @@
 #include "MediaViewerUtils.h"
 
 #include "CanvasTypes.h"
+#include "ColorManagement/TransferFunctions.h"
 #include "DetailsViewArgs.h"
 #include "EditorClassUtils.h"
 #include "Engine/Canvas.h"
-#include "Engine/Engine.h"
-#include "Engine/TextureRenderTarget2D.h"
 #include "IStructureDetailsView.h"
 #include "Materials/Material.h"
 #include "Materials/MaterialInterface.h"
@@ -156,11 +155,11 @@ TOptional<TVariant<FColor, FLinearColor>> FMediaViewerUtils::GetPixelColor(TArra
 	return ColorVariant;
 }
 
-UTextureRenderTarget2D* FMediaViewerUtils::CreateRenderTarget(const FIntPoint& InSize, bool bInTransparent)
+UTextureRenderTarget2D* FMediaViewerUtils::CreateRenderTarget(const FIntPoint& InSize, bool bInTransparent, ETextureRenderTargetFormat InFormat)
 {
 	UTextureRenderTarget2D* RenderTarget = NewObject<UTextureRenderTarget2D>(GetTransientPackage());
 	check(RenderTarget);
-	RenderTarget->RenderTargetFormat = RTF_RGBA8;
+	RenderTarget->RenderTargetFormat = InFormat;
 	RenderTarget->bAutoGenerateMips = false;
 	RenderTarget->bCanCreateUAV = false;
 	RenderTarget->ClearColor = bInTransparent ? FLinearColor::Transparent : FLinearColor::Black;
@@ -232,4 +231,20 @@ TSharedRef<IStructureDetailsView> FMediaViewerUtils::CreateStructDetailsView(TSh
 	return PropertyEditorModule.CreateStructureDetailView(DetailsArgs, StructArgs, InStructOnScope, InCustomName);
 }
 
+FColor FMediaViewerUtils::CorrectGamma(const FColor& InColor)
+{
+	return CorrectGamma(InColor.ReinterpretAsLinear()).ToFColor(/* SRGB */ false);
+}
+
+FLinearColor FMediaViewerUtils::CorrectGamma(const FLinearColor& InColor)
+{
+	FLinearColor Color;
+	Color.R = UE::Color::EncodeSRGB(InColor.R);
+	Color.G = UE::Color::EncodeSRGB(InColor.G);
+	Color.B = UE::Color::EncodeSRGB(InColor.B);
+	Color.A = InColor.A;
+
+	return Color;
+}
+
 } // UE::MediaViewer::Private
diff --git a/Engine/Plugins/Media/MediaViewer/Source/MediaViewer/Private/MediaViewerUtils.h b/Engine/Plugins/Media/MediaViewer/Source/MediaViewer/Private/MediaViewerUtils.h
--- a/Engine/Plugins/Media/MediaViewer/Source/MediaViewer/Private/MediaViewerUtils.h
+++ b/Engine/Plugins/Media/MediaViewer/Source/MediaViewer/Private/MediaViewerUtils.h
@@ -6,6 +6,7 @@
 
 #include "Containers/ArrayView.h"
 #include "DetailsViewArgs.h"
+#include "Engine/TextureRenderTarget2D.h"
 #include "HAL/Platform.h"
 #include "IStructureDetailsView.h"
 #include "Math/Color.h"
@@ -55,7 +56,7 @@ public:
 	 * Creates a render target texture parented to the transient package. 
 	 * The target will have the @see UMediaViewerUserData asset data to differentiate it from other render targets.
 	 */
-	static UTextureRenderTarget2D* CreateRenderTarget(const FIntPoint& InSize, bool bInTransparent);
+	static UTextureRenderTarget2D* CreateRenderTarget(const FIntPoint& InSize, bool bInTransparent, ETextureRenderTargetFormat InFormat = ETextureRenderTargetFormat::RTF_RGBA8);
 
 	/** 
 	 * Creates a render target (@see CreateRenderTarget) and renders the material to it.
@@ -69,6 +70,20 @@ public:
 	/** Creates a struct details view based on the given struct. The view will have most settings disabled for a clean view. */
 	static TSharedRef<IStructureDetailsView> CreateStructDetailsView(TSharedRef<FStructOnScope> InStructOnScope, const FText& InCustomName,
 		FNotifyHook* InNotifyHook = nullptr);
+
+	/**
+	 * Corrects the gamma levels in the given color.
+	 * 
+	 * @return The corrected color.
+	 */
+	static FColor CorrectGamma(const FColor& InColor);
+
+	/**
+	 * Corrects the gamma levels in the given color.
+	 *
+	 * @return The corrected color.
+	 */
+	static FLinearColor CorrectGamma(const FLinearColor& InColor);
 };
 
 } // UE::MediaViewer::Private
diff --git a/Engine/Plugins/Media/MediaViewer/Source/MediaViewer/Private/Widgets/SMediaImageViewerOverlay.cpp b/Engine/Plugins/Media/MediaViewer/Source/MediaViewer/Private/Widgets/SMediaImageViewerOverlay.cpp
--- a/Engine/Plugins/Media/MediaViewer/Source/MediaViewer/Private/Widgets/SMediaImageViewerOverlay.cpp
+++ b/Engine/Plugins/Media/MediaViewer/Source/MediaViewer/Private/Widgets/SMediaImageViewerOverlay.cpp
@@ -505,6 +505,21 @@ void SMediaImageViewerOverlay::BindCommands()
 		FIsActionButtonVisible::CreateSP(this, &SMediaImageViewerOverlay::CanCopyTransform)
 	);
 
+	CommandList->MapAction(
+		Commands.MirrorHorizonal,
+		FExecuteAction::CreateSP(this, &SMediaImageViewerOverlay::Try_Mirror, EMediaImageMirrorPlane::Horizontal)
+	);
+
+	CommandList->MapAction(
+		Commands.MirrorVertical,
+		FExecuteAction::CreateSP(this, &SMediaImageViewerOverlay::Try_Mirror, EMediaImageMirrorPlane::Vertical)
+	);
+
+	CommandList->MapAction(
+		Commands.MirrorDepth,
+		FExecuteAction::CreateSP(this, &SMediaImageViewerOverlay::Try_Mirror, EMediaImageMirrorPlane::Depth)
+	);
+
 	CommandList->MapAction(Commands.CopyColor, FExecuteAction::CreateSP(this, &SMediaImageViewerOverlay::CopyColor));
 
 	CommandList->MapAction(
@@ -984,6 +999,29 @@ void SMediaImageViewerOverlay::CopyTransform()
 	Delegates->SetTransformToAll.Execute(GetOffset(), GetRotation(), GetScale());
 }
 
+void SMediaImageViewerOverlay::Try_Mirror(EMediaImageMirrorPlane InPlane)
+{
+	if (TSharedPtr<FMediaImageViewer> ImageViewer = Delegates->GetImageViewer.Execute(Position))
+	{
+		if (Delegates->AreTransformsLocked.Execute())
+		{
+			Delegates->MirrorAll.Execute(InPlane, !ImageViewer->IsMirrored(InPlane));
+		}
+		else
+		{
+			Mirror(InPlane);
+		}
+	}
+}
+
+void SMediaImageViewerOverlay::Mirror(EMediaImageMirrorPlane InPlane)
+{
+	if (TSharedPtr<FMediaImageViewer> ImageViewer = Delegates->GetImageViewer.Execute(Position))
+	{
+		ImageViewer->SetMirrored(InPlane, !ImageViewer->IsMirrored(InPlane));
+	}
+}
+
 uint8 SMediaImageViewerOverlay::GetMipLevel() const
 {
 	if (TSharedPtr<FMediaImageViewer> ImageViewer = Delegates->GetImageViewer.Execute(Position))
diff --git a/Engine/Plugins/Media/MediaViewer/Source/MediaViewer/Private/Widgets/SMediaImageViewerOverlay.h b/Engine/Plugins/Media/MediaViewer/Source/MediaViewer/Private/Widgets/SMediaImageViewerOverlay.h
--- a/Engine/Plugins/Media/MediaViewer/Source/MediaViewer/Private/Widgets/SMediaImageViewerOverlay.h
+++ b/Engine/Plugins/Media/MediaViewer/Source/MediaViewer/Private/Widgets/SMediaImageViewerOverlay.h
@@ -8,6 +8,7 @@ class FUICommandList;
 class SOverlay;
 class SVerticalBox;
 enum class ECheckBoxState : uint8;
+enum class EMediaImageMirrorPlane : uint8;
 struct FMediaViewerLibraryItem;
 
 namespace UE::MediaViewer
@@ -162,6 +163,10 @@ protected:
 	bool CanCopyTransform();
 	void CopyTransform();
 
+	/** Interact with the image viewer's mirror settings. */
+	void Try_Mirror(EMediaImageMirrorPlane InPlane);
+	void Mirror(EMediaImageMirrorPlane InPlane);
+
 	/** Interact with the image viewer's mip level. */
 	uint8 GetMipLevel() const;
 	void AdjustMipLevel(int8 InAdjustment);
diff --git a/Engine/Plugins/Media/MediaViewer/Source/MediaViewer/Private/Widgets/SMediaImageViewerStatusBar.cpp b/Engine/Plugins/Media/MediaViewer/Source/MediaViewer/Private/Widgets/SMediaImageViewerStatusBar.cpp
--- a/Engine/Plugins/Media/MediaViewer/Source/MediaViewer/Private/Widgets/SMediaImageViewerStatusBar.cpp
+++ b/Engine/Plugins/Media/MediaViewer/Source/MediaViewer/Private/Widgets/SMediaImageViewerStatusBar.cpp
@@ -8,6 +8,7 @@
 #include "Internationalization/Text.h"
 #include "MediaViewerStyle.h"
 #include "Styling/StyleColors.h"
+#include "Widgets/Colors/SColorBlock.h"
 #include "Widgets/MediaImageStatusBarExtender.h"
 #include "Widgets/SBoxPanel.h"
 #include "Widgets/SNullWidget.h"
@@ -107,6 +108,18 @@ TSharedRef<SWidget> SMediaImageViewerStatusBar::BuildStatusBar()
 				.Text(this, &SMediaImageViewerStatusBar::GetColorPickerLabel)
 				.DecoratorStyleSet(&FMediaViewerStyle::Get())
 			];
+
+		HorizontalBox->AddSlot()
+			.Padding(SlotPadding)
+			.AutoWidth()
+			.HAlign(HAlign_Right)
+			.VAlign(VAlign_Center)
+			[
+				SNew(SColorBlock)
+				.Size(FVector2D(12))
+				.Color(this, &SMediaImageViewerStatusBar::GetColorBlockColor)
+				.UseSRGB(false)
+			];
 	}
 	ApplyHook(StatusBarSections::StatusBarRight, EExtensionHook::After);
 
@@ -197,6 +210,51 @@ FText SMediaImageViewerStatusBar::GetColorPickerLabel() const
 	return FText::GetEmpty();
 }
 
+FLinearColor SMediaImageViewerStatusBar::GetColorBlockColor() const
+{
+	TSharedPtr<FMediaImageViewer> ImageViewer = Delegates->GetImageViewer.Execute(Position);
+
+	if (!ImageViewer.IsValid() || !ImageViewer->IsValid())
+	{
+		return FLinearColor::Black;
+	}
+
+	const FIntPoint PixelCoordinates = Delegates->GetPixelCoordinates.Execute(Position);
+
+	if (PixelCoordinates.X < 0 || PixelCoordinates.Y < 0)
+	{
+		return FLinearColor::Black;
+	}
+
+	const FIntPoint& ImageSize = ImageViewer->GetInfo().Size;
+
+	if (PixelCoordinates.X >= ImageSize.X || PixelCoordinates.Y >= ImageSize.Y)
+	{
+		return FLinearColor::Black;
+	}
+
+	const TOptional<TVariant<FColor, FLinearColor>> PixelColor = ImageViewer->GetPixelColor(PixelCoordinates, ImageViewer->GetPaintSettings().GetMipLevel());
+
+	if (!PixelColor.IsSet())
+	{
+		return FLinearColor::Black;
+	}
+
+	const TVariant<FColor, FLinearColor>& PixelColorValue = PixelColor.GetValue();
+
+	if (const FLinearColor* ColorLinear = PixelColorValue.TryGet<FLinearColor>())
+	{
+		return *ColorLinear;
+	}
+
+	if (const FColor* Color = PixelColorValue.TryGet<FColor>())
+	{
+		return Color->ReinterpretAsLinear();
+	}
+
+	return FLinearColor::Black;
+}
+
 } // UE::MediaViewer::Private
 
 #undef LOCTEXT_NAMESPACE
diff --git a/Engine/Plugins/Media/MediaViewer/Source/MediaViewer/Private/Widgets/SMediaImageViewerStatusBar.h b/Engine/Plugins/Media/MediaViewer/Source/MediaViewer/Private/Widgets/SMediaImageViewerStatusBar.h
--- a/Engine/Plugins/Media/MediaViewer/Source/MediaViewer/Private/Widgets/SMediaImageViewerStatusBar.h
+++ b/Engine/Plugins/Media/MediaViewer/Source/MediaViewer/Private/Widgets/SMediaImageViewerStatusBar.h
@@ -41,6 +41,8 @@ protected:
 	FText GetResolutionLabel() const;
 
 	FText GetColorPickerLabel() const;
+
+	FLinearColor GetColorBlockColor() const;
 };
 
 } // UE::MediaViewer::Private
diff --git a/Engine/Plugins/Media/MediaViewer/Source/MediaViewer/Private/Widgets/SMediaViewer.cpp b/Engine/Plugins/Media/MediaViewer/Source/MediaViewer/Private/Widgets/SMediaViewer.cpp
--- a/Engine/Plugins/Media/MediaViewer/Source/MediaViewer/Private/Widgets/SMediaViewer.cpp
+++ b/Engine/Plugins/Media/MediaViewer/Source/MediaViewer/Private/Widgets/SMediaViewer.cpp
@@ -45,8 +45,6 @@
 
 #define LOCTEXT_NAMESPACE "SMediaViewer"
 
-extern ENGINE_API UEngine* GEngine;
-
 namespace UE::MediaViewer::Private
 {
 
@@ -584,6 +582,7 @@ void SMediaViewer::CreateDelegates()
 	Delegates->MultiplyScaleAroundCursorToAll.BindSP(this, &SMediaViewer::MultiplyScaleAroundCursorToAll);
 	Delegates->SetTransformToAll.BindSP(this, &SMediaViewer::SetTransformToAll);
 	Delegates->ResetTransformToAll.BindSP(this, &SMediaViewer::ResetTransformToAll);
+	Delegates->MirrorAll.BindSP(this, &SMediaViewer::MirrorAll);
 	Delegates->GetSecondImageViewerOpacity.BindSPLambda(this, [this] { return MediaViewerSettings.SecondImageOpacity; });
 	Delegates->SetSecondImageViewerOpacity.BindSPLambda(this, [this](float InValue) { MediaViewerSettings.SecondImageOpacity = InValue; });
 	Delegates->GetABSplitterLocation.BindSPLambda(this, [this] { return MediaViewerSettings.ABSplitterLocation; });
@@ -1015,6 +1014,14 @@ void SMediaViewer::CopyTransformToAll(EMediaImageViewerPosition InPosition) cons
 	}
 }
 
+void SMediaViewer::MirrorAll(EMediaImageMirrorPlane InPlane, bool bInSet)
+{
+	for (int32 Index = 0; Index < static_cast<int32>(EMediaImageViewerPosition::COUNT); ++Index)
+	{
+		ImageViewers[Index]->SetMirrored(InPlane, bInSet);
+	}
+}
+
 TSharedPtr<FUICommandList> SMediaViewer::GetOverlayCommandList(EMediaImageViewerPosition InPosition) const
 {
 	if (TSharedPtr<SMediaImageViewerOverlay> Overlay = GetOverlay(InPosition))
diff --git a/Engine/Plugins/Media/MediaViewer/Source/MediaViewer/Private/Widgets/SMediaViewer.h b/Engine/Plugins/Media/MediaViewer/Source/MediaViewer/Private/Widgets/SMediaViewer.h
--- a/Engine/Plugins/Media/MediaViewer/Source/MediaViewer/Private/Widgets/SMediaViewer.h
+++ b/Engine/Plugins/Media/MediaViewer/Source/MediaViewer/Private/Widgets/SMediaViewer.h
@@ -19,6 +19,7 @@ class SViewportToolBar;
 class UCanvas;
 class UPackage;
 class UWorld;
+enum class EMediaImageMirrorPlane : uint8;
 struct FMediaViewerState;
 
 namespace UE::MediaViewer
@@ -224,6 +225,9 @@ protected:
 	/** Copies the transform of the given image viewer to all other image viewers. */
 	void CopyTransformToAll(EMediaImageViewerPosition InPosition) const;
 
+	/** Sets the mirror plane flag on all image viewers. */
+	void MirrorAll(EMediaImageMirrorPlane InPlane, bool bInSet);
+
 	/**
 	 * Returns the coordinates of the given image that the mousee is currently hovering.
 	 * May not be inside the image.
diff --git a/Engine/Plugins/Media/MediaViewer/Source/MediaViewer/Public/ImageViewer/MediaImageViewer.h b/Engine/Plugins/Media/MediaViewer/Source/MediaViewer/Public/ImageViewer/MediaImageViewer.h
--- a/Engine/Plugins/Media/MediaViewer/Source/MediaViewer/Public/ImageViewer/MediaImageViewer.h
+++ b/Engine/Plugins/Media/MediaViewer/Source/MediaViewer/Public/ImageViewer/MediaImageViewer.h
@@ -69,6 +69,20 @@ struct FMediaImageViewerPanelSettings
 	TSoftObjectPtr<UTexture> BackgroundTexture;
 };
 
+UENUM(meta = (Bitflags, UseEnumValuesAsMaskValuesInEditor = "true"))
+enum class EMediaImageMirrorPlane : uint8
+{
+	None       = 0      UMETA(Hidden),
+	Horizontal = 1 << 0,
+	Vertical   = 1 << 1,
+	Depth      = 1 << 2 UMETA(ToolTip = "3D Only"), // Front to back flipping
+
+	MAX        = Depth UMETA(Hidden),
+	TwoDim     = Horizontal | Vertical UMETA(Hidden),
+	ThreeDim   = TwoDim | Depth UMETA(Hidden)
+};
+ENUM_CLASS_FLAGS(EMediaImageMirrorPlane)
+
 USTRUCT()
 struct FMediaImagePaintSettings
 {
@@ -94,6 +108,12 @@ struct FMediaImagePaintSettings
 	UPROPERTY(EditAnywhere, Config, Category = "Media")
 	float Scale = 1.f;
 
+	/**
+	 * Image mirroring along the center axes
+	 */
+	UPROPERTY(EditAnywhere, Config, Category = "Media", meta = (Bitmask, BitmaskEnum = "/Script/MediaViewer.EMediaImageMirrorPlane"))
+	uint8 MirrorPlanes = 0;
+
 	/** Whether to enable the MipOverride */
 	UPROPERTY(EditAnywhere, Config, Category = "Media", meta = (InlineEditConditionToggle))
 	bool bEnableMipOverride = false;
@@ -325,6 +345,15 @@ public:
 	 */
 	virtual TSharedPtr<FStructOnScope> GetCustomSettingsOnScope() const { return nullptr; }
 
+	/** Returns true if the image is mirrored on the given axis. */
+	bool IsMirrored(EMediaImageMirrorPlane InPlane) const;
+
+	/** Use to set mirror flag, don't set directly! */
+	void SetMirrored(EMediaImageMirrorPlane InPlane, bool bInSet);
+
+	/** Triggered when a shape is mirrored. */
+	virtual void OnMirrored() {};
+
 	/** Returns the center of the viewer, not necessarily the center of the splitter it's contained within. */
 	MEDIAVIEWER_API FVector2D GetViewerCenter(const FVector2D& InViewerSize) const;
 
diff --git a/Engine/Plugins/Media/MediaViewer/Source/MediaViewer/Public/MediaViewerDelegates.h b/Engine/Plugins/Media/MediaViewer/Source/MediaViewer/Public/MediaViewerDelegates.h
--- a/Engine/Plugins/Media/MediaViewer/Source/MediaViewer/Public/MediaViewerDelegates.h
+++ b/Engine/Plugins/Media/MediaViewer/Source/MediaViewer/Public/MediaViewerDelegates.h
@@ -10,6 +10,7 @@
 #include "UObject/StructOnScope.h"
 
 class FUICommandList;
+enum class EMediaImageMirrorPlane : uint8;
 enum class EMediaImageViewerActivePosition : uint8;
 enum EOrientation : int;
 struct FMediaViewerSettings;
@@ -43,6 +44,7 @@ struct FMediaViewerDelegates
 	DECLARE_DELEGATE_RetVal(FMediaViewerSettings*, FGetSettings)
 	DECLARE_DELEGATE_RetVal(TSharedRef<IMediaViewerLibrary>, FGetLibrary)
 	DECLARE_DELEGATE_RetVal(TSharedPtr<FUICommandList>, FGetCommandList)
+	DECLARE_DELEGATE_TwoParams(FSetAllMirrorPlaneFlags, EMediaImageMirrorPlane, bool /* bInSet */)
 
 	DECLARE_DELEGATE_RetVal_OneParam(TSharedPtr<FMediaImageViewer>, FGetImageViewer, EMediaImageViewerPosition)
 	DECLARE_DELEGATE_TwoParams(FSetImageViewer, EMediaImageViewerPosition, const TSharedRef<FMediaImageViewer>&)
@@ -110,6 +112,9 @@ struct FMediaViewerDelegates
 	/** Resets the transform of all image viewers to 0,0,0; 0,0,0; 1. */
 	FSimpleDelegate ResetTransformToAll;
 
+	/** Set the mirror plane flag (on/off) to all image viewers. */
+	FSetAllMirrorPlaneFlags MirrorAll;
+
 	/** Gets the opacity of the second/B image viewer. */
 	FGetFloat GetSecondImageViewerOpacity;
 
