Horde issue 1186474

#Synced-CL 49593614

diff --git a/Engine/Plugins/Media/MediaViewer/Source/MediaViewer/Private/ImageViewers/MediaImageViewer.cpp b/Engine/Plugins/Media/MediaViewer/Source/MediaViewer/Private/ImageViewers/MediaImageViewer.cpp
--- a/Engine/Plugins/Media/MediaViewer/Source/MediaViewer/Private/ImageViewers/MediaImageViewer.cpp
+++ b/Engine/Plugins/Media/MediaViewer/Source/MediaViewer/Private/ImageViewers/MediaImageViewer.cpp
@@ -103,7 +103,7 @@ void FMediaImageViewer::SetMirrored(EMediaImageMirrorPlane InPlane, bool bInSet)
 	}
 	else
 	{
-		PaintSettings.MirrorPlanes = PaintSettings.MirrorPlanes & !static_cast<uint8>(InPlane);
+		PaintSettings.MirrorPlanes = PaintSettings.MirrorPlanes & ~static_cast<uint8>(InPlane);
 	}
 
 	OnMirrored();
