Media Viewer: All the viewport tabs can now have their viewports enumerated.

#Synced-CL 49784486

diff --git a/Engine/Plugins/Media/MediaViewer/Source/MediaViewer/Private/ImageViewers/LevelEditorViewportImageViewer.cpp b/Engine/Plugins/Media/MediaViewer/Source/MediaViewer/Private/ImageViewers/LevelEditorViewportImageViewer.cpp
--- a/Engine/Plugins/Media/MediaViewer/Source/MediaViewer/Private/ImageViewers/LevelEditorViewportImageViewer.cpp
+++ b/Engine/Plugins/Media/MediaViewer/Source/MediaViewer/Private/ImageViewers/LevelEditorViewportImageViewer.cpp
@@ -16,24 +16,73 @@
 namespace UE::MediaViewer::Private
 {
 
-/** Can take either a string value (Viewport0..3) or a config key and return the string value (Viewport0..3) */
-FString ConfigKeyToStringValue(const FString& InConfigKey)
+namespace LevelEditorViewport
 {
-	int32 DotIndex;
-	InConfigKey.FindLastChar('.', DotIndex);
 
-	if (DotIndex < 0 || DotIndex >= (InConfigKey.Len() - 1))
+	const FGuid ViewportGuids[4] = {
+		FGuid(0xFDD736C4, 0xFD7A98CB, 0x7C28F833, 0xD00B4BB3),
+		FGuid(0xB7929682, 0x749CED27, 0x357D74D5, 0xC5E19053),
+		FGuid(0xD82E7C42, 0xB038A8A1, 0x9C608700, 0x9C3B589C),
+		FGuid(0xF0FE7988, 0xE5DE498B, 0x4AC6026B, 0x255B3BEC)
+	};
+
+	int32 GetTabIndex(const FString& InValue)
 	{
-		return InConfigKey;
+		if (InValue.Len() > 1 && InValue[1] == TEXT('_')
+			&& InValue[0] >= TEXT('0') && InValue[0] <= TEXT('3'))
+		{
+			return InValue[0] - TEXT('0');
+		}
+
+		return 0;
 	}
 
-	return InConfigKey.Mid(DotIndex + 1);
+	int32 GetViewportIndex(const FString& InValue)
+	{
+		const int32 LastIndex = InValue.Len() - 1;
+
+		for (int32 Index = LastIndex; Index >= 0; --Index)
+		{
+			if (InValue[Index] >= TEXT('0') && InValue[Index] <= TEXT('9'))
+			{
+				continue;
+			}
+
+			// No numbers
+			if (Index == LastIndex)
+			{
+				return 0;
+			}
+
+			// We've found a non-number
+			return FCString::Atoi(*InValue.Mid(Index + 1));
+		}
+
+		// Whole string is a number... somehow!
+		return FCString::Atoi(*InValue);
+	}
+}
+
+/** 
+ * Can take either a string value (Viewport0..3) or a config key and return the string value (Viewport0..3) 
+ * Now also processed tab index (X_...)
+ */
+FString ConfigKeyToStringValue(const FString& InConfigKey)
+{
+	return InConfigKey;
 }
 
 /** Could be a config key or a direct value (Viewport0..3) */
 FText GetViewportDisplayName(const FString& InViewportValue)
 {
-	return FText::FromString(ConfigKeyToStringValue(InViewportValue));
+	using namespace LevelEditorViewport;
+
+	const int32 TabIndex = GetTabIndex(InViewportValue);
+	const int32 ViewportIndex = GetViewportIndex(InViewportValue);
+
+	return FText::FromString(FString(TEXT("Tab")) + FString::FromInt(TabIndex + 1)
+		+ TEXT(":")
+		+ FString(TEXT("Viewport")) + FString::FromInt(ViewportIndex + 1));
 }
 
 FText GetViewportTooltip(const FString& InViewportValue)
@@ -57,22 +106,20 @@ FGuid FLevelEditorViewportImageViewer::FItem::GetIdForViewport(const FString& In
 {
 	if (!InConfigKey.IsEmpty())
 	{
-		const TCHAR LastChar = InConfigKey[InConfigKey.Len() - 1];
+		using namespace LevelEditorViewport;
 
-		switch (LastChar)
-		{
-			case '0':
-				return FGuid(0xFDD736C4, 0xFD7A98CB, 0x7C28F833, 0xD00B4BB3);
+		const int32 TabIndex = FMath::Clamp(GetTabIndex(InConfigKey), 0, 3);
+		const int32 ViewportIndex = FMath::Clamp(GetViewportIndex(InConfigKey), 0, 3);
 
-			case '1':
-				return FGuid(0xB7929682, 0x749CED27, 0x357D74D5, 0xC5E19053);
+		if (TabIndex == 0)
+		{
+			return ViewportGuids[ViewportIndex];
+		}
 
-			case '2':
-				return FGuid(0xD82E7C42, 0xB038A8A1, 0x9C608700, 0x9C3B589C);
+		FGuid ViewportGuid = ViewportGuids[ViewportIndex];
+		ViewportGuid.D += TabIndex;
 
-			case '3':
-				return FGuid(0xF0FE7988, 0xE5DE498B, 0x4AC6026B, 0x255B3BEC);
-		}
+		return ViewportGuid;
 	}
 
 	if (bCreateIfInvalid)
@@ -104,19 +151,24 @@ TSharedPtr<FSceneViewport> FLevelEditorViewportImageViewer::FItem::GetViewportFr
 		return nullptr;
 	}
 
-	TArray<TSharedPtr<SLevelViewport>> Viewports = LevelEditor->GetViewports();
+	TMap<int32, TArray<TSharedPtr<SLevelViewport>>> Viewports = LevelEditor->GetViewportsByTab();
+
+	using namespace LevelEditorViewport;
 
-	for (const TSharedPtr<SLevelViewport>& Viewport : Viewports)
+	const int32 TabIndex = GetTabIndex(InStringValue);
+	const int32 ViewportIndex = GetViewportIndex(InStringValue);
+
+	if (!Viewports.Contains(TabIndex))
 	{
-		const FString ViewportValue = ConfigKeyToStringValue(Viewport->GetConfigKey().ToString());
+		return nullptr;
+	}
 
-		if (ViewportValue == InStringValue)
-		{
-			return Viewport->GetSharedActiveViewport();
-		}
+	if (!Viewports[TabIndex].IsValidIndex(ViewportIndex))
+	{
+		return nullptr;
 	}
 
-	return nullptr;
+	return Viewports[TabIndex][ViewportIndex]->GetSharedActiveViewport();
 }
 
 FLevelEditorViewportImageViewer::FItem::FItem(const FString& InConfigKey)
diff --git a/Engine/Plugins/Media/MediaViewer/Source/MediaViewer/Private/Library/LevelEditorViewportGroup.cpp b/Engine/Plugins/Media/MediaViewer/Source/MediaViewer/Private/Library/LevelEditorViewportGroup.cpp
--- a/Engine/Plugins/Media/MediaViewer/Source/MediaViewer/Private/Library/LevelEditorViewportGroup.cpp
+++ b/Engine/Plugins/Media/MediaViewer/Source/MediaViewer/Private/Library/LevelEditorViewportGroup.cpp
@@ -44,45 +44,48 @@ TArray<TSharedRef<FMediaViewerLibraryItem>> FLevelEditorViewportGroup::GetLevelE
 		return LevelEditorViewportItems;
 	}
 
-	TArray<TSharedPtr<SLevelViewport>> Viewports = LevelEditor->GetViewports();
+	TMap<int32, TArray<TSharedPtr<SLevelViewport>>> Viewports = LevelEditor->GetViewportsByTab();
 	
 	TSet<FGuid> AddedViewportIds;
 	AddedViewportIds.Reserve(Viewports.Num());
 
-	for (int32 Index = 0; Index < Viewports.Num(); ++Index)
+	for (const TPair<int32, TArray<TSharedPtr<SLevelViewport>>>& ViewportTabPair : Viewports)
 	{
-		TSharedPtr<FSceneViewport> ActiveViewport = Viewports[Index]->GetSharedActiveViewport();
-
-		if (!ActiveViewport.IsValid())
+		for (int32 Index = 0; Index < ViewportTabPair.Value.Num(); ++Index)
 		{
-			continue;
-		}
+			TSharedPtr<FSceneViewport> ActiveViewport = ViewportTabPair.Value[Index]->GetSharedActiveViewport();
 
-		const FIntPoint Size = ActiveViewport->GetSize();
+			if (!ActiveViewport.IsValid())
+			{
+				continue;
+			}
 
-		if (Size.X < 2 || Size.Y < 2)
-		{
-			continue;
-		}
+			const FIntPoint Size = ActiveViewport->GetSize();
 
-		const FString ConfigKey = Viewports[Index]->GetConfigKey().ToString();
+			if (Size.X < 2 || Size.Y < 2)
+			{
+				continue;
+			}
 
-		const FGuid ViewportId = FLevelEditorViewportImageViewer::FItem::GetIdForViewport(
-			ConfigKey,
-			/* Create id if invalid */ false
-		);
+			const FString ConfigKey = FString::FromInt(ViewportTabPair.Key) + TEXT("_") + ViewportTabPair.Value[Index]->GetConfigKey().ToString();
 
-		bool bAlreadyInSet = false;
-		AddedViewportIds.Add(ViewportId, &bAlreadyInSet);
+			const FGuid ViewportId = FLevelEditorViewportImageViewer::FItem::GetIdForViewport(
+				ConfigKey,
+				/* Create id if invalid */ false
+			);
 
-		if (bAlreadyInSet)
-		{
-			continue;
-		}
+			bool bAlreadyInSet = false;
+			AddedViewportIds.Add(ViewportId, &bAlreadyInSet);
 
-		if (ViewportId.IsValid())
-		{
-			LevelEditorViewportItems.Add(MakeShared<FLevelEditorViewportImageViewer::FItem>(ViewportId, ConfigKey));
+			if (bAlreadyInSet)
+			{
+				continue;
+			}
+
+			if (ViewportId.IsValid())
+			{
+				LevelEditorViewportItems.Add(MakeShared<FLevelEditorViewportImageViewer::FItem>(ViewportId, ConfigKey));
+			}
 		}
 	}
 
