Media Viewer: Return of the Side Bar!
- Mirror commands are now available on the right click menu.

#Synced-CL 49947617

diff --git a/Engine/Plugins/Media/MediaViewer/Source/MediaViewer/Private/Widgets/SMediaImageViewerOverlay.cpp b/Engine/Plugins/Media/MediaViewer/Source/MediaViewer/Private/Widgets/SMediaImageViewerOverlay.cpp
--- a/Engine/Plugins/Media/MediaViewer/Source/MediaViewer/Private/Widgets/SMediaImageViewerOverlay.cpp
+++ b/Engine/Plugins/Media/MediaViewer/Source/MediaViewer/Private/Widgets/SMediaImageViewerOverlay.cpp
@@ -507,17 +507,23 @@ void SMediaImageViewerOverlay::BindCommands()
 
 	CommandList->MapAction(
 		Commands.MirrorHorizonal,
-		FExecuteAction::CreateSP(this, &SMediaImageViewerOverlay::Try_Mirror, EMediaImageMirrorPlane::Horizontal)
+		FExecuteAction::CreateSP(this, &SMediaImageViewerOverlay::Try_Mirror, EMediaImageMirrorPlane::Horizontal),
+		FCanExecuteAction(),
+		FGetActionCheckState::CreateSP(this, &SMediaImageViewerOverlay::GetMirroredState, EMediaImageMirrorPlane::Horizontal)
 	);
 
 	CommandList->MapAction(
 		Commands.MirrorVertical,
-		FExecuteAction::CreateSP(this, &SMediaImageViewerOverlay::Try_Mirror, EMediaImageMirrorPlane::Vertical)
+		FExecuteAction::CreateSP(this, &SMediaImageViewerOverlay::Try_Mirror, EMediaImageMirrorPlane::Vertical),
+		FCanExecuteAction(),
+		FGetActionCheckState::CreateSP(this, &SMediaImageViewerOverlay::GetMirroredState, EMediaImageMirrorPlane::Vertical)
 	);
 
 	CommandList->MapAction(
 		Commands.MirrorDepth,
-		FExecuteAction::CreateSP(this, &SMediaImageViewerOverlay::Try_Mirror, EMediaImageMirrorPlane::Depth)
+		FExecuteAction::CreateSP(this, &SMediaImageViewerOverlay::Try_Mirror, EMediaImageMirrorPlane::Depth),
+		FCanExecuteAction(),
+		FGetActionCheckState::CreateSP(this, &SMediaImageViewerOverlay::GetMirroredState, EMediaImageMirrorPlane::Depth)
 	);
 
 	CommandList->MapAction(Commands.CopyColor, FExecuteAction::CreateSP(this, &SMediaImageViewerOverlay::CopyColor));
@@ -1014,6 +1020,18 @@ void SMediaImageViewerOverlay::Try_Mirror(EMediaImageMirrorPlane InPlane)
 	}
 }
 
+ECheckBoxState SMediaImageViewerOverlay::GetMirroredState(EMediaImageMirrorPlane InPlane) const
+{
+	if (TSharedPtr<FMediaImageViewer> ImageViewer = Delegates->GetImageViewer.Execute(Position))
+	{
+		return ImageViewer->IsMirrored(InPlane)
+			? ECheckBoxState::Checked
+			: ECheckBoxState::Unchecked;
+	}
+
+	return ECheckBoxState::Undetermined;
+}
+
 void SMediaImageViewerOverlay::Mirror(EMediaImageMirrorPlane InPlane)
 {
 	if (TSharedPtr<FMediaImageViewer> ImageViewer = Delegates->GetImageViewer.Execute(Position))
@@ -1128,17 +1146,25 @@ TSharedRef<SWidget> SMediaImageViewerOverlay::CreateMenu()
 
 	FMenuBuilder MenuBuilder(true, CommandList, nullptr, false, &FAppStyle::Get(), false);
 
-	MenuBuilder.BeginSection(TEXT("Image"), LOCTEXT("Image", "Image"));
+	if (GetImageViewer() != FNullImageViewer::GetNullImageViewer())
 	{
-		MenuBuilder.AddMenuEntry(Commands.ResetTransform);
-		MenuBuilder.AddMenuEntry(Commands.CopyTransform);
-		MenuBuilder.AddMenuEntry(Commands.AddToLibrary);
-		MenuBuilder.AddSeparator();
-		MenuBuilder.AddMenuEntry(Commands.ToggleMipOverride);
-		MenuBuilder.AddMenuEntry(Commands.IncreaseMipLevel);
-		MenuBuilder.AddMenuEntry(Commands.DecreaseMipLevel);
+		MenuBuilder.BeginSection(TEXT("Image"), LOCTEXT("Image", "Image"));
+		{
+			MenuBuilder.AddMenuEntry(Commands.ResetTransform);
+			MenuBuilder.AddMenuEntry(Commands.CopyTransform);
+			MenuBuilder.AddMenuEntry(Commands.AddToLibrary);
+			MenuBuilder.AddSeparator();
+			MenuBuilder.AddMenuEntry(Commands.ToggleMipOverride);
+			MenuBuilder.AddMenuEntry(Commands.IncreaseMipLevel);
+			MenuBuilder.AddMenuEntry(Commands.DecreaseMipLevel);
+			MenuBuilder.AddSeparator();
+			MenuBuilder.AddMenuEntry(Commands.MirrorHorizonal);
+			MenuBuilder.AddMenuEntry(Commands.MirrorVertical);
+			// This is disabled until 3D image viewers are added.
+			//MenuBuilder.AddMenuEntry(Commands.MirrorDepth);
+		}
+		MenuBuilder.EndSection();
 	}
-	MenuBuilder.EndSection();
 
 	MenuBuilder.BeginSection(TEXT("Media Viewer"), LOCTEXT("MediaViewer", "Media Viewer"));
 	{
diff --git a/Engine/Plugins/Media/MediaViewer/Source/MediaViewer/Private/Widgets/SMediaImageViewerOverlay.h b/Engine/Plugins/Media/MediaViewer/Source/MediaViewer/Private/Widgets/SMediaImageViewerOverlay.h
--- a/Engine/Plugins/Media/MediaViewer/Source/MediaViewer/Private/Widgets/SMediaImageViewerOverlay.h
+++ b/Engine/Plugins/Media/MediaViewer/Source/MediaViewer/Private/Widgets/SMediaImageViewerOverlay.h
@@ -71,6 +71,9 @@ public:
 	/** Scales the image to fit the viewer. */
 	void ScaleToFit(bool bInUseTransformLock);
 
+	/** Either mirror this viewer or all viewers if transforms are locked. */
+	void Try_Mirror(EMediaImageMirrorPlane InPlane);
+
 	//~ Begin SWidget
 	virtual FCursorReply OnCursorQuery(const FGeometry& InMyGeometry, const FPointerEvent& InCursorEvent) const override;
 	virtual FReply OnKeyDown(const FGeometry& InMyGeometry, const FKeyEvent& InKeyEvent) override;
@@ -164,7 +167,7 @@ protected:
 	void CopyTransform();
 
 	/** Interact with the image viewer's mirror settings. */
-	void Try_Mirror(EMediaImageMirrorPlane InPlane);
+	ECheckBoxState GetMirroredState(EMediaImageMirrorPlane InPlane) const;
 	void Mirror(EMediaImageMirrorPlane InPlane);
 
 	/** Interact with the image viewer's mip level. */
diff --git a/Engine/Plugins/Media/MediaViewer/Source/MediaViewer/Private/Widgets/SMediaViewer.cpp b/Engine/Plugins/Media/MediaViewer/Source/MediaViewer/Private/Widgets/SMediaViewer.cpp
--- a/Engine/Plugins/Media/MediaViewer/Source/MediaViewer/Private/Widgets/SMediaViewer.cpp
+++ b/Engine/Plugins/Media/MediaViewer/Source/MediaViewer/Private/Widgets/SMediaViewer.cpp
@@ -27,6 +27,7 @@
 #include "MediaViewerUtils.h"
 #include "Misc/MessageDialog.h"
 #include "PackageTools.h"
+#include "Sidebar/ISidebarDrawerContent.h"
 #include "Sidebar/SSidebar.h"
 #include "Sidebar/SSidebarContainer.h"
 #include "TextureResource.h"
@@ -48,6 +49,55 @@
 namespace UE::MediaViewer::Private
 {
 
+static const FLazyName DrawerId = "Library";
+
+class FLibraryDrawerConfig : public ISidebarDrawerContent
+{
+public:
+	FLibraryDrawerConfig(TSharedRef<SMediaViewerLibraryPrivate> InLibraryWidget)
+		: LibraryWidget(InLibraryWidget)
+	{
+	}
+
+	virtual FName GetUniqueId() const override
+	{
+		return DrawerId;
+	}
+
+	virtual FName GetSectionId() const
+	{
+		return DrawerId;
+	}
+
+	virtual FText GetSectionDisplayText() const
+	{
+		return LOCTEXT("Library", "Library");
+	}
+
+	virtual bool ShouldShowSection() const
+	{
+		return true;
+	}
+
+	virtual TOptional<float> GetSectionFillHeight() const
+	{
+		return 1.f;
+	}
+
+	virtual int32 GetSortOrder() const
+	{
+		return 0;
+	}
+
+	virtual TSharedRef<SWidget> CreateContentWidget()
+	{
+		return LibraryWidget;
+	}
+
+private:
+	TSharedRef<SMediaViewerLibraryPrivate> LibraryWidget;
+};
+
 const FSlateColorBrush SMediaViewer::BackgroundColorBrush(FLinearColor::Black);
 
 SMediaViewer::SMediaViewer()
@@ -144,17 +194,15 @@ void SMediaViewer::Construct(const FArguments& InArgs, const TSharedRef<SMediaVi
 
 	SidebarContainer->RebuildSidebar(Sidebar, SidebarState);
 
-	const FName DrawerId = "Library";
-
 	FSidebarDrawerConfig LibraryDrawerConfig;
 	LibraryDrawerConfig.UniqueId = DrawerId;
 	LibraryDrawerConfig.ButtonText = LOCTEXT("Library", "Library");
 	LibraryDrawerConfig.ToolTipText = LOCTEXT("LibraryTooltip", "Open the Library side bar.");
 	LibraryDrawerConfig.Icon = FAppStyle::GetBrush("Icons.FolderOpen");
-	LibraryDrawerConfig.InitialState = SidebarState.FindOrAddDrawerState(DrawerId);
-	LibraryDrawerConfig.OverrideContentWidget = Library;
+	LibraryDrawerConfig.InitialState = SidebarState.FindOrAddDrawerState(DrawerId.Resolve());
 
 	Sidebar->RegisterDrawer(MoveTemp(LibraryDrawerConfig));
+	Sidebar->RegisterDrawerSection(DrawerId, MakeShared<FLibraryDrawerConfig>(Library.ToSharedRef()));
 
 	ChildSlot
 	[
@@ -250,6 +298,21 @@ void SMediaViewer::BindCommands()
 
 	CommandList->MapAction(Commands.OpenNewTab, FExecuteAction::CreateSP(this, &SMediaViewer::OpenNewTab));
 
+	CommandList->MapAction(
+		Commands.MirrorHorizonal,
+		FExecuteAction::CreateSP(this, &SMediaViewer::Mirror, EMediaImageMirrorPlane::Horizontal)
+	);
+
+	CommandList->MapAction(
+		Commands.MirrorVertical,
+		FExecuteAction::CreateSP(this, &SMediaViewer::Mirror, EMediaImageMirrorPlane::Vertical)
+	);
+
+	CommandList->MapAction(
+		Commands.MirrorDepth,
+		FExecuteAction::CreateSP(this, &SMediaViewer::Mirror, EMediaImageMirrorPlane::Depth)
+	);
+
 	auto BindBookmark = [this](const TSharedPtr<FUICommandInfo>& InOpenCommand, const TSharedPtr<FUICommandInfo>& InSaveCommand, int32 InIndex)
 		{
 			CommandList->MapAction(
@@ -1448,6 +1511,20 @@ void SMediaViewer::OnAssetsDeleted(const TArray<UObject*>& InAssets)
 	}
 }
 
+void SMediaViewer::Mirror(EMediaImageMirrorPlane InPlane)
+{
+	const FGuid& NullImageId = FNullImageViewer::GetNullImageViewer()->GetInfo().Id;
+
+	for (int32 Index = 0; Index < static_cast<int32>(EMediaImageViewerPosition::COUNT); ++Index)
+	{
+		if (TSharedPtr<SMediaImageViewerOverlay> Overlay = GetOverlay(static_cast<EMediaImageViewerPosition>(Index)))
+		{
+			Overlay->Try_Mirror(InPlane);
+			break;
+		}
+	}
+}
+
 } // UE::MediaViewer::Private
 
 #undef LOCTEXT_NAMESPACE
diff --git a/Engine/Plugins/Media/MediaViewer/Source/MediaViewer/Private/Widgets/SMediaViewer.h b/Engine/Plugins/Media/MediaViewer/Source/MediaViewer/Private/Widgets/SMediaViewer.h
--- a/Engine/Plugins/Media/MediaViewer/Source/MediaViewer/Private/Widgets/SMediaViewer.h
+++ b/Engine/Plugins/Media/MediaViewer/Source/MediaViewer/Private/Widgets/SMediaViewer.h
@@ -276,6 +276,9 @@ protected:
 	void OnPackageDeleted(UPackage* InPackage);
 
 	void OnAssetsDeleted(const TArray<UObject*>& InAssets);
+
+	/** Forwards the mirror bind to the first valid viewer. */
+	void Mirror(EMediaImageMirrorPlane InPlane);
 };
 
 } // UE::MediaViewer::Private
