Media Viewer: Render targets no longer render without gamma.

#Synced-CL 49990324

diff --git a/Engine/Plugins/Media/MediaViewer/Source/MediaViewer/Private/ImageViewers/TextureRenderTarget2DImageViewer.cpp b/Engine/Plugins/Media/MediaViewer/Source/MediaViewer/Private/ImageViewers/TextureRenderTarget2DImageViewer.cpp
--- a/Engine/Plugins/Media/MediaViewer/Source/MediaViewer/Private/ImageViewers/TextureRenderTarget2DImageViewer.cpp
+++ b/Engine/Plugins/Media/MediaViewer/Source/MediaViewer/Private/ImageViewers/TextureRenderTarget2DImageViewer.cpp
@@ -195,7 +195,7 @@ FTextureRenderTarget2DImageViewer::FTextureRenderTarget2DImageViewer(const FGuid
 		FMediaImageViewer::GetObjectDisplayName(InRenderTarget)
 	})
 {
-	DrawEffects |= ESlateDrawEffect::PreMultipliedAlpha | ESlateDrawEffect::NoGamma;
+	DrawEffects |= ESlateDrawEffect::PreMultipliedAlpha;
 
 	RenderTargetSettings.RenderTarget = InRenderTarget;
 
diff --git a/Engine/Plugins/Media/MediaViewer/Source/MediaViewer/Private/ImageViewers/TextureRenderTarget2DImageViewer.h b/Engine/Plugins/Media/MediaViewer/Source/MediaViewer/Private/ImageViewers/TextureRenderTarget2DImageViewer.h
--- a/Engine/Plugins/Media/MediaViewer/Source/MediaViewer/Private/ImageViewers/TextureRenderTarget2DImageViewer.h
+++ b/Engine/Plugins/Media/MediaViewer/Source/MediaViewer/Private/ImageViewers/TextureRenderTarget2DImageViewer.h
@@ -32,7 +32,7 @@ public:
 	{
 		FFactory()
 		{
-			Priority = 5000;
+			Priority = 4000;
 		}
 
 		virtual bool SupportsAsset(const FAssetData& InAssetData) const override;
