Media Viewer: The clear color default is now black instead of dark grey. The clear color was affecting the color of the displayed media.

#Synced-CL 50388552

diff --git a/Engine/Plugins/Media/MediaViewer/Source/MediaViewer/Private/Widgets/MediaViewerSettings.h b/Engine/Plugins/Media/MediaViewer/Source/MediaViewer/Private/Widgets/MediaViewerSettings.h
--- a/Engine/Plugins/Media/MediaViewer/Source/MediaViewer/Private/Widgets/MediaViewerSettings.h
+++ b/Engine/Plugins/Media/MediaViewer/Source/MediaViewer/Private/Widgets/MediaViewerSettings.h
@@ -37,7 +37,7 @@ struct FMediaViewerSettings
 
 	/** Color painted to the viewer behind the images. */
 	UPROPERTY(EditAnywhere, Config, Category = "Background")
-	FLinearColor ClearColor = FStyleColors::Recessed.GetSpecifiedColor();
+	FLinearColor ClearColor = FStyleColors::Black.GetSpecifiedColor();
 
 	/** Texture tiled behind the images, on top of the clear color. */
 	UPROPERTY(EditAnywhere, Config, Category = "Background")
