MediaViewer: Accept drops from image sequences.

#rnx

#Synced-CL 50633722

diff --git a/Engine/Plugins/Media/MediaViewer/Source/MediaViewer/MediaViewer.build.cs b/Engine/Plugins/Media/MediaViewer/Source/MediaViewer/MediaViewer.build.cs
--- a/Engine/Plugins/Media/MediaViewer/Source/MediaViewer/MediaViewer.build.cs
+++ b/Engine/Plugins/Media/MediaViewer/Source/MediaViewer/MediaViewer.build.cs
@@ -40,6 +40,7 @@ public class MediaViewer : ModuleRules
 				"RHI",
 				"SlateCore",
 				"Slate",
+				"TimeManagement",
 				"ToolWidgets",
 				"WorkspaceMenuStructure",
 				"UMG",
diff --git a/Engine/Plugins/Media/MediaViewer/Source/MediaViewer/Private/Widgets/SMediaViewerDropTarget.cpp b/Engine/Plugins/Media/MediaViewer/Source/MediaViewer/Private/Widgets/SMediaViewerDropTarget.cpp
--- a/Engine/Plugins/Media/MediaViewer/Source/MediaViewer/Private/Widgets/SMediaViewerDropTarget.cpp
+++ b/Engine/Plugins/Media/MediaViewer/Source/MediaViewer/Private/Widgets/SMediaViewerDropTarget.cpp
@@ -4,13 +4,18 @@
 
 #include "AssetRegistry/AssetRegistryModule.h"
 #include "AssetSelection.h"
+#include "CommonFrameRates.h"
 #include "Containers/Array.h"
 #include "Containers/ArrayView.h"
 #include "ContentBrowserModule.h"
 #include "DragAndDrop/DecoratedDragDropOp.h"
+#include "Editor.h"
 #include "FileMediaSource.h"
 #include "Framework/Application/SlateApplication.h"
+#include "HAL/FileManager.h"
 #include "IContentBrowserSingleton.h"
+#include "IMediaEventSink.h"
+#include "IMediaStreamPlayer.h"
 #include "ImageViewers/MediaSourceImageViewer.h"
 #include "IMediaViewerModule.h"
 #include "ImgMediaSource.h"
@@ -18,19 +23,266 @@
 #include "Library/MediaViewerLibrary.h"
 #include "Library/MediaViewerLibraryItem.h"
 #include "MediaPlayer.h"
+#include "MediaStream.h"
 #include "MediaViewerDelegates.h"
 #include "MediaViewerModule.h"
 #include "Misc/MessageDialog.h"
 #include "SDropTarget.h"
 #include "Styling/AppStyle.h"
 #include "Styling/StyleColors.h"
+#include "Widgets/Input/SButton.h"
 #include "Widgets/Layout/SBox.h"
+#include "Widgets/Layout/SBorder.h"
+#include "Widgets/Layout/SUniformGridPanel.h"
 #include "Widgets/MediaViewerLibraryItemDragDropOperation.h"
 #include "Widgets/SBoxPanel.h"
+#include "Widgets/SFrameRatePicker.h"
 #include "Widgets/SOverlay.h"
+#include "Widgets/SWindow.h"
+#include "Widgets/Text/STextBlock.h"
 
 #define LOCTEXT_NAMESPACE "SMediaViewerDropTarget"
 
+namespace
+{
+
+/** Modal dialog that lets the user pick a frame rate for an image sequence. */
+class SImageSequenceFrameRateDialog : public SCompoundWidget
+{
+public:
+	SLATE_BEGIN_ARGS(SImageSequenceFrameRateDialog) {}
+	SLATE_END_ARGS()
+
+	void Construct(const FArguments& InArgs, TSharedRef<SWindow> InWindow, const FFrameRate& InDefaultFrameRate)
+	{
+		Window = InWindow;
+		SelectedFrameRate = InDefaultFrameRate;
+
+		ChildSlot
+		[
+			SNew(SBorder)
+				.Visibility(EVisibility::Visible)
+				.BorderImage(FAppStyle::GetBrush("Menu.Background"))
+				[
+					SNew(SVerticalBox)
+
+					+ SVerticalBox::Slot()
+						.FillHeight(1)
+						.VAlign(VAlign_Top)
+						[
+							SNew(SBorder)
+								.BorderImage(FAppStyle::GetBrush("ToolPanel.GroupBorder"))
+								.Padding(4.0f)
+								[
+									SNew(SVerticalBox)
+
+									+ SVerticalBox::Slot()
+										.Padding(0.0f, 0.0f, 0.0f, 4.0f)
+										[
+											SNew(STextBlock)
+												.Text(LOCTEXT("SelectFrameRateLabel", "Select a frame rate for this image sequence:"))
+										]
+
+									+ SVerticalBox::Slot()
+										.AutoHeight()
+										[
+											SNew(SFrameRatePicker)
+												.Value(this, &SImageSequenceFrameRateDialog::GetFrameRate)
+												.OnValueChanged(this, &SImageSequenceFrameRateDialog::SetFrameRate)
+										]
+								]
+						]
+
+					+ SVerticalBox::Slot()
+						.AutoHeight()
+						.HAlign(HAlign_Right)
+						.VAlign(VAlign_Bottom)
+						.Padding(8)
+						[
+							SNew(SUniformGridPanel)
+								.SlotPadding(FAppStyle::GetMargin("StandardDialog.SlotPadding"))
+								.MinDesiredSlotWidth(FAppStyle::GetFloat("StandardDialog.MinDesiredSlotWidth"))
+								.MinDesiredSlotHeight(FAppStyle::GetFloat("StandardDialog.MinDesiredSlotHeight"))
+
+							+ SUniformGridPanel::Slot(0, 0)
+								[
+									SNew(SButton)
+										.HAlign(HAlign_Center)
+										.ContentPadding(FAppStyle::GetMargin("StandardDialog.ContentPadding"))
+										.OnClicked_Lambda([this]() -> FReply
+										{
+											CloseDialog(true);
+											return FReply::Handled();
+										})
+										.Text(LOCTEXT("OkButtonLabel", "OK"))
+								]
+
+							+ SUniformGridPanel::Slot(1, 0)
+								[
+									SNew(SButton)
+										.HAlign(HAlign_Center)
+										.ContentPadding(FAppStyle::GetMargin("StandardDialog.ContentPadding"))
+										.OnClicked_Lambda([this]() -> FReply
+										{
+											CloseDialog(false);
+											return FReply::Handled();
+										})
+										.Text(LOCTEXT("CancelButtonLabel", "Cancel"))
+								]
+						]
+				]
+		];
+	}
+
+	bool WasConfirmed() const
+	{
+		return bConfirmed;
+	}
+
+	FFrameRate GetFrameRate() const
+	{
+		return SelectedFrameRate;
+	}
+
+private:
+	void SetFrameRate(FFrameRate InFrameRate)
+	{
+		SelectedFrameRate = InFrameRate;
+	}
+
+	void CloseDialog(bool bInConfirmed)
+	{
+		bConfirmed = bInConfirmed;
+
+		if (TSharedPtr<SWindow> PinnedWindow = Window.Pin())
+		{
+			PinnedWindow->RequestDestroyWindow();
+		}
+	}
+
+	TWeakPtr<SWindow> Window;
+	FFrameRate SelectedFrameRate;
+	bool bConfirmed = false;
+};
+
+/**
+ * Converts a float frame rate to the nearest FFrameRate from the common frame rate presets.
+ * Falls back to FFrameRate(30, 1) if no match is found.
+ */
+FFrameRate FindClosestCommonFrameRate(float InRate)
+{
+	FFrameRate BestMatch(30, 1);
+
+	if (InRate <= 0.0f)
+	{
+		return BestMatch;
+	}
+
+	float BestDiff = FMath::Abs(BestMatch.AsDecimal() - static_cast<double>(InRate));
+
+	for (const FCommonFrameRateInfo& Info : FCommonFrameRates::GetAll())
+	{
+		const float Diff = FMath::Abs(Info.FrameRate.AsDecimal() - static_cast<double>(InRate));
+
+		if (Diff < BestDiff)
+		{
+			BestDiff = Diff;
+			BestMatch = Info.FrameRate;
+		}
+	}
+
+	return BestMatch;
+}
+
+/**
+ * Extracts a trailing frame number from a filename, ignoring the file extension.
+ * For example, "render.0042.exr" yields 42, "frame_001.png" yields 1.
+ * Based on the same logic as FImgMediaLoader::GetNumberAtEndOfString().
+ */
+bool GetFrameNumberFromFilename(const FString& InFilename, const FString& InExtension, int32& OutFrameNumber)
+{
+	// Start scanning from the right, skipping the extension (and its dot).
+	int32 Index = InFilename.Len() - 1 - InExtension.Len();
+
+	// Skip the dot before the extension.
+	if (Index >= 0 && InFilename[Index] == TEXT('.'))
+	{
+		--Index;
+	}
+
+	// Find the first digit from the right.
+	for (; Index >= 0; --Index)
+	{
+		if (FChar::IsDigit(InFilename[Index]))
+		{
+			break;
+		}
+	}
+
+	if (Index < 0)
+	{
+		return false;
+	}
+
+	const int32 LastDigitIndex = Index;
+
+	// Find the start of the digit sequence.
+	for (; Index >= 0; --Index)
+	{
+		if (!FChar::IsDigit(InFilename[Index]))
+		{
+			break;
+		}
+	}
+
+	++Index;
+
+	const FString NumberString = InFilename.Mid(Index, LastDigitIndex - Index + 1);
+	OutFrameNumber = FCString::Atoi(*NumberString);
+
+	return true;
+}
+
+/**
+ * Returns true if the file has an extension registered as an image media source.
+ * Uses SpawnMediaSourceForString to check whether the extension produces a UImgMediaSource,
+ * with results cached per extension to avoid repeated object creation.
+ */
+bool HasSupportedImageExtension(const FString& InFilePath)
+{
+	static TSet<FString> SupportedExtensions;
+	static TSet<FString> UnsupportedExtensions;
+
+	const FString Extension = FPaths::GetExtension(InFilePath).ToLower();
+
+	if (SupportedExtensions.Contains(Extension))
+	{
+		return true;
+	}
+
+	if (UnsupportedExtensions.Contains(Extension))
+	{
+		return false;
+	}
+
+	// Probe by creating a source via the registered delegate and checking its type.
+	UMediaSource* TestSource = UMediaSource::SpawnMediaSourceForString(InFilePath, GetTransientPackage());
+	const bool bIsImageSource = TestSource != nullptr && TestSource->IsA<UImgMediaSource>();
+
+	if (bIsImageSource)
+	{
+		SupportedExtensions.Add(Extension);
+	}
+	else
+	{
+		UnsupportedExtensions.Add(Extension);
+	}
+
+	return bIsImageSource;
+}
+
+} // anonymous namespace
+
 namespace UE::MediaViewer::Private
 {
 
@@ -302,25 +554,32 @@ bool SMediaViewerDropTarget::CreateAndSetImageSequence(const FExternalDragOperat
 		return false;
 	}
 
+	int32 DroppedFrameNumber = 0;
+	bool bSingleFrameDrop = false;
+
 	// Drag and drop a directory - open as image sequence
 	if (Files.Num() == 1 && FPaths::DirectoryExists(Files[0]))
 	{
 		ImageSequencePath = Files[0];
 	}
+	// Drag and drop a single image file (standalone or part of a sequence)
+	else if (Files.Num() == 1 && HasSupportedImageExtension(Files[0]))
+	{
+		ImageSequencePath = FPaths::GetPath(Files[0]);
+		bSingleFrameDrop = true;
+	}
 	// Drag and drop more than 2 files (2 would be a comparison view)
 	else if (Files.Num() > 2)
 	{
 		ImageSequencePath = FPaths::GetPath(Files[0]);
 
-		const TArray<FString> ValidTypes = { TEXT("BMP"), TEXT("EXR"), TEXT("PNG"), TEXT("JPG"), TEXT("JPEG") };
-
 		// Sanity check
 		if (!FPaths::DirectoryExists(ImageSequencePath))
 		{
 			return false;
 		}
 
-		if (!ValidTypes.Contains(FPaths::GetExtension(Files[0]).ToUpper()))
+		if (!HasSupportedImageExtension(Files[0]))
 		{
 			return false;
 		}
@@ -333,7 +592,7 @@ bool SMediaViewerDropTarget::CreateAndSetImageSequence(const FExternalDragOperat
 				return false;
 			}
 
-			if (!ValidTypes.Contains(FPaths::GetExtension(Files[Index]).ToUpper()))
+			if (!HasSupportedImageExtension(Files[Index]))
 			{
 				return false;
 			}
@@ -344,6 +603,134 @@ bool SMediaViewerDropTarget::CreateAndSetImageSequence(const FExternalDragOperat
 		return false;
 	}
 
+	if (bSingleFrameDrop)
+	{
+		// Determine if the file is part of a numbered image sequence or a standalone image.
+		const bool bIsSequenceFrame = IsImageSequenceFile(Files[0], DroppedFrameNumber);
+
+		const FString Extension = FPaths::GetExtension(Files[0]).ToLower();
+		const int32 FirstFrameNumber = bIsSequenceFrame ? GetFirstFrameNumber(ImageSequencePath, Extension) : 0;
+		const int32 SeekFrame = DroppedFrameNumber - FirstFrameNumber;
+
+		// Use the registered spawn delegate to create the appropriate source type.
+		// This verifies the extension is actually handled by ImgMedia (not a video format).
+		UMediaSource* MediaSource = UMediaSource::SpawnMediaSourceForString(Files[0], GetTransientPackage());
+		UImgMediaSource* ImgSource = Cast<UImgMediaSource>(MediaSource);
+
+		if (!ImgSource)
+		{
+			return false;
+		}
+
+		ImgSource->ClearFlags(RF_Transactional);
+		ImgSource->SetFlags(RF_Transient);
+
+		const FText DisplayName = FText::FromString(FPaths::GetCleanFilename(ImageSequencePath));
+		TSharedRef<FMediaSourceImageViewer> ImageViewer = MakeShared<FMediaSourceImageViewer>(ImgSource, DisplayName);
+
+		if (UMediaStream* MediaStream = ImageViewer->GetMediaStream())
+		{
+			if (IMediaStreamPlayer* StreamPlayer = MediaStream->GetPlayer().GetInterface())
+			{
+				// Disable auto-play so the viewer opens paused.
+				FMediaStreamPlayerConfig Config = StreamPlayer->GetPlayerConfig();
+				Config.bPlayOnOpen = false;
+				StreamPlayer->SetPlayerConfig(Config);
+
+				// For sequence frames, show the frame rate dialog and seek after media opens.
+				if (bIsSequenceFrame)
+				{
+					if (UMediaPlayer* Player = StreamPlayer->GetPlayer())
+					{
+						TWeakObjectPtr<UMediaPlayer> WeakPlayer = Player;
+						TWeakObjectPtr<UImgMediaSource> WeakSource = ImgSource;
+
+						// Track whether the dialog has been shown so we don't repeat it on re-open.
+						TSharedRef<bool> bDialogShown = MakeShared<bool>(false);
+
+						Player->OnMediaEvent().AddLambda([WeakPlayer, WeakSource, SeekFrame, bDialogShown](EMediaEvent Event)
+						{
+							if (Event != EMediaEvent::MediaOpened)
+							{
+								return;
+							}
+
+							UMediaPlayer* StrongPlayer = WeakPlayer.Get();
+
+							if (!StrongPlayer)
+							{
+								return;
+							}
+
+							if (!*bDialogShown)
+							{
+								*bDialogShown = true;
+
+								// Query the frame rate that the reader detected from the file.
+								const float DetectedRate = StrongPlayer->GetVideoTrackFrameRate(INDEX_NONE, INDEX_NONE);
+								const FFrameRate DefaultFrameRate = FindClosestCommonFrameRate(DetectedRate);
+
+								FFrameRate SelectedFrameRate;
+
+								if (!ShowFrameRatePickerDialog(SelectedFrameRate, DefaultFrameRate))
+								{
+									// User cancelled - leave the viewer showing the first frame.
+									return;
+								}
+
+								// If the user changed the rate, re-open with the override applied.
+								if (!FMath::IsNearlyEqual(SelectedFrameRate.AsDecimal(), DefaultFrameRate.AsDecimal(), 0.01))
+								{
+									if (UImgMediaSource* Source = WeakSource.Get())
+									{
+										Source->FrameRateOverride = SelectedFrameRate;
+										StrongPlayer->Close();
+										StrongPlayer->OpenSource(Source);
+
+										// The second MediaOpened will handle seeking below.
+										return;
+									}
+								}
+
+								// Same rate as detected - seek to the dropped frame.
+								if (SeekFrame > 0)
+								{
+									const double SeekSeconds = static_cast<double>(SeekFrame) / SelectedFrameRate.AsDecimal();
+									StrongPlayer->Seek(FTimespan::FromSeconds(SeekSeconds));
+								}
+							}
+							else
+							{
+								// Second open after a frame rate change - seek to the dropped frame.
+								if (SeekFrame > 0)
+								{
+									if (UImgMediaSource* Source = WeakSource.Get())
+									{
+										const double SeekSeconds = static_cast<double>(SeekFrame) / Source->FrameRateOverride.AsDecimal();
+										StrongPlayer->Seek(FTimespan::FromSeconds(SeekSeconds));
+									}
+								}
+							}
+						});
+					}
+				}
+			}
+		}
+
+		// Swap images and then replace first image.
+		if (bForceComparisonView && Position == EMediaImageViewerPosition::First)
+		{
+			Delegates->SwapAB.Execute();
+			Delegates->SetABView.Execute();
+			Delegates->SetABOrientation.Execute(Orient_Horizontal);
+		}
+
+		Delegates->SetImageViewer.Execute(Position, ImageViewer);
+
+		return true;
+	}
+
+	// Directory or multi-file drop: existing save-to-asset workflow.
 	const EAppReturnType::Type Return = FMessageDialog::Open(
 		EAppMsgType::OkCancel,
 		LOCTEXT("ImportAsImageSequence", "Import as Image Media Source?")
@@ -392,7 +779,7 @@ bool SMediaViewerDropTarget::CreateAndSetImageSequence(const FExternalDragOperat
 
 	IMediaViewerModule& MediaViewerModule = IMediaViewerModule::Get();
 	TSharedPtr<FMediaViewerLibraryItem> LibraryItem = MediaViewerModule.CreateLibraryItem(TNotNull<UObject*>(ImageSequence));
-	
+
 	if (!LibraryItem.IsValid())
 	{
 		return false;
@@ -416,10 +803,82 @@ bool SMediaViewerDropTarget::CreateAndSetImageSequence(const FExternalDragOperat
 	}
 
 	Delegates->SetImageViewer.Execute(Position, ImageViewer.ToSharedRef());
-	
+
 	return true;
 }
 
+bool SMediaViewerDropTarget::IsImageSequenceFile(const FString& InFilePath, int32& OutFrameNumber)
+{
+	if (!HasSupportedImageExtension(InFilePath))
+	{
+		return false;
+	}
+
+	const FString Extension = FPaths::GetExtension(InFilePath).ToLower();
+	const FString Filename = FPaths::GetCleanFilename(InFilePath);
+
+	return GetFrameNumberFromFilename(Filename, Extension, OutFrameNumber);
+}
+
+int32 SMediaViewerDropTarget::GetFirstFrameNumber(const FString& InDirectoryPath, const FString& InExtension)
+{
+	TArray<FString> FoundFiles;
+	IFileManager::Get().FindFiles(FoundFiles, *FPaths::Combine(InDirectoryPath, TEXT("*")), true, false);
+
+	int32 MinFrameNumber = 0;
+	bool bFoundAny = false;
+
+	for (const FString& File : FoundFiles)
+	{
+		if (FPaths::GetExtension(File).ToLower() != InExtension)
+		{
+			continue;
+		}
+
+		int32 FrameNumber = 0;
+
+		if (GetFrameNumberFromFilename(File, InExtension, FrameNumber))
+		{
+			if (!bFoundAny || FrameNumber < MinFrameNumber)
+			{
+				MinFrameNumber = FrameNumber;
+				bFoundAny = true;
+			}
+		}
+	}
+
+	return MinFrameNumber;
+}
+
+bool SMediaViewerDropTarget::ShowFrameRatePickerDialog(FFrameRate& OutFrameRate, const FFrameRate& InDefaultFrameRate)
+{
+	if (!GEditor)
+	{
+		return false;
+	}
+
+	TSharedRef<SWindow> Window = SNew(SWindow)
+		.Title(LOCTEXT("ImageSequenceFrameRateTitle", "Image Sequence Frame Rate"))
+		.ClientSize(FVector2D(350, 130))
+		.SupportsMinimize(false)
+		.SupportsMaximize(false);
+
+	TSharedRef<SImageSequenceFrameRateDialog> Dialog =
+		SNew(SImageSequenceFrameRateDialog, Window, InDefaultFrameRate);
+
+	Window->SetContent(Dialog);
+
+	GEditor->EditorAddModalWindow(Window);
+
+	if (Dialog->WasConfirmed())
+	{
+		OutFrameRate = Dialog->GetFrameRate();
+		return true;
+	}
+
+	return false;
+}
+
 bool SMediaViewerDropTarget::OnAllowDrop(TSharedPtr<FDragDropOperation> InDragDropOperation) const
 {
 	if (!InDragDropOperation.IsValid())
diff --git a/Engine/Plugins/Media/MediaViewer/Source/MediaViewer/Private/Widgets/SMediaViewerDropTarget.h b/Engine/Plugins/Media/MediaViewer/Source/MediaViewer/Private/Widgets/SMediaViewerDropTarget.h
--- a/Engine/Plugins/Media/MediaViewer/Source/MediaViewer/Private/Widgets/SMediaViewerDropTarget.h
+++ b/Engine/Plugins/Media/MediaViewer/Source/MediaViewer/Private/Widgets/SMediaViewerDropTarget.h
@@ -4,6 +4,7 @@
 
 #include "Widgets/SCompoundWidget.h"
 
+#include "Misc/FrameRate.h"
 #include "SMediaViewer.h"
 
 namespace UE::MediaViewer
@@ -83,6 +84,37 @@ protected:
 	 * Returns true if an asset was created.
 	 */
 	bool CreateAndSetImageSequence(const FExternalDragOperation& InFileDragDropOp);
+
+	/**
+	 * Checks if the given file is a supported image sequence frame.
+	 * Looks for a known image extension and a trailing frame number in the filename.
+	 * 
+	 * @param InFilePath The full path to the file.
+	 * @param OutFrameNumber If successful, the extracted frame number.
+	 * 
+	 * @return True if the file looks like an image sequence frame.
+	 */
+	static bool IsImageSequenceFile(const FString& InFilePath, int32& OutFrameNumber);
+
+	/**
+	 * Scans a directory for image files with the given extension and returns the smallest frame number found.
+	 * 
+	 * @param InDirectoryPath The directory to scan.
+	 * @param InExtension The file extension to filter by (exclude '.').
+	 * 
+	 * @return The minimum frame number, or 0 if none found.
+	 */
+	static int32 GetFirstFrameNumber(const FString& InDirectoryPath, const FString& InExtension);
+
+	/**
+	 * Shows a modal dialog to let the user pick a frame rate for an image sequence.
+	 * 
+	 * @param OutFrameRate Receives the selected frame rate.
+	 * @param InDefaultFrameRate The frame rate to pre-select in the dialog.
+	 * 
+	 * @return True if the user confirmed.
+	 */
+	static bool ShowFrameRatePickerDialog(FFrameRate& OutFrameRate, const FFrameRate& InDefaultFrameRate = FFrameRate(30, 1));
 };
 
 } // UE::MediaViewer::Private
