Media Viewer: Improved the display of mip levels.

#Synced-CL 50687911

diff --git a/Engine/Plugins/Media/MediaViewer/Source/MediaViewer/Private/ImageViewers/MaterialInterfaceImageViewer.cpp b/Engine/Plugins/Media/MediaViewer/Source/MediaViewer/Private/ImageViewers/MaterialInterfaceImageViewer.cpp
--- a/Engine/Plugins/Media/MediaViewer/Source/MediaViewer/Private/ImageViewers/MaterialInterfaceImageViewer.cpp
+++ b/Engine/Plugins/Media/MediaViewer/Source/MediaViewer/Private/ImageViewers/MaterialInterfaceImageViewer.cpp
@@ -412,6 +412,8 @@ void FMaterialInterfaceImageViewer::RenderMaterial()
 {
 	if (::IsValid(MaterialSettings.MaterialInterface) && ::IsValid(MaterialSettings.RenderTarget))
 	{
+		bInvalidatedMipTexture = true;
+
 		FMediaViewerUtils::RenderMaterial(MaterialSettings.MaterialInterface, MaterialSettings.RenderTarget);
 
 		if (SampleCache.IsValid())
diff --git a/Engine/Plugins/Media/MediaViewer/Source/MediaViewer/Private/ImageViewers/MediaImageViewer.cpp b/Engine/Plugins/Media/MediaViewer/Source/MediaViewer/Private/ImageViewers/MediaImageViewer.cpp
--- a/Engine/Plugins/Media/MediaViewer/Source/MediaViewer/Private/ImageViewers/MediaImageViewer.cpp
+++ b/Engine/Plugins/Media/MediaViewer/Source/MediaViewer/Private/ImageViewers/MediaImageViewer.cpp
@@ -61,6 +61,12 @@ void FMediaImageViewer::UpdateId(const FGuid& InId)
 void FMediaImageViewer::NotifyPostChange(const FPropertyChangedEvent& InPropertyChangedEvent, FProperty* InPropertyThatChanged)
 {
 	InitBackgroundTexture();
+
+	if (InPropertyChangedEvent.GetMemberPropertyName() == GET_MEMBER_NAME_CHECKED(FMediaImagePaintSettings, MipOverride))
+	{
+		// This should already be invalidated by the mip level changing the required render texture size, but we'll make sure.
+		bInvalidatedMipTexture = true;
+	}
 }
 
 void FMediaImageViewer::AddReferencedObjects(FReferenceCollector& InCollector)
@@ -152,39 +158,63 @@ bool FMediaImageViewer::UpdateMipBrush()
 		{
 			return false;
 		}
+
+		bInvalidatedMipTexture = true;
 	}
 
+	// Create a render target that is the same size as the mip level.
+	const FVector2D RequiredSizeDouble = 
+		FVector2D(ImageInfo.Size.X, ImageInfo.Size.Y)
+		/ FMath::Max(
+			1.0, 
+			static_cast<float>(FMath::Pow(2.0, PaintSettings.GetMipLevel())
+		)
+	);
+
+	const FIntPoint RequiredSize = {
+		static_cast<int32>(FMath::CeilToInt(RequiredSizeDouble.X)),
+		static_cast<int32>(FMath::CeilToInt(RequiredSizeDouble.Y))
+	};
+	
 	if (!MipAdjustedRenderTarget.IsValid())
 	{
-		MipAdjustedRenderTarget.Reset(Private::FMediaViewerUtils::CreateRenderTarget(ImageInfo.Size, /* Transparent */ true));
+		MipAdjustedRenderTarget.Reset(Private::FMediaViewerUtils::CreateRenderTarget(RequiredSize, /* Transparent */ true));
 
 		if (!MipAdjustedRenderTarget.IsValid())
 		{
 			return false;
 		}
+
+		MipAdjustedRenderTarget->Filter = TextureFilter::TF_Nearest;
+		bInvalidatedMipTexture = true;
 	}
 	// If it was just created, we don't need to change the size.
-	else if (MipAdjustedRenderTarget->SizeX != ImageInfo.Size.X || MipAdjustedRenderTarget->SizeY != ImageInfo.Size.Y)
+	else if (MipAdjustedRenderTarget->SizeX != RequiredSize.X || MipAdjustedRenderTarget->SizeY != RequiredSize.Y)
 	{
-		MipAdjustedRenderTarget->ResizeTarget(ImageInfo.Size.X, ImageInfo.Size.Y);
+		MipAdjustedRenderTarget->ResizeTarget(RequiredSize.X, RequiredSize.Y);
+		bInvalidatedMipTexture = true;
 	}
 
 	if (!MipAdjustedBrush.IsValid())
 	{
-		const FVector2D MaterialBrushSize = FVector2D(ImageInfo.Size.X, ImageInfo.Size.Y);
-		MipAdjustedBrush = MakeShared<FSlateMaterialBrush>(*MipAdjustedMaterial.Get(), MaterialBrushSize);
+		MipAdjustedBrush = MakeShared<FSlateImageBrush>(MipAdjustedRenderTarget.Get(), RequiredSizeDouble);
 	}
 	// If it was just created, we don't need to change the size.
-	else if (!FMath::IsNearlyEqual(MipAdjustedBrush->ImageSize.X, ImageInfo.Size.X)
-		|| !FMath::IsNearlyEqual(MipAdjustedBrush->ImageSize.Y, ImageInfo.Size.Y))
+	else if (!FMath::IsNearlyEqual(MipAdjustedBrush->ImageSize.X, RequiredSize.X)
+		|| !FMath::IsNearlyEqual(MipAdjustedBrush->ImageSize.Y, RequiredSize.Y))
 	{
-		const FVector2D MaterialBrushSize = FVector2D(ImageInfo.Size.X, ImageInfo.Size.Y);
-		MipAdjustedBrush->SetImageSize(MaterialBrushSize);
+		MipAdjustedBrush->SetImageSize(RequiredSizeDouble);
 	}
 
 	MipAdjustedMaterial->SetTextureParameterValue(TEXT("Texture"), BrushTexture);
 	MipAdjustedMaterial->SetScalarParameterValue(TEXT("MipLevel"), static_cast<float>(PaintSettings.GetMipLevel()));
 
+	if (bInvalidatedMipTexture)
+	{
+		Private::FMediaViewerUtils::RenderMaterial(MipAdjustedMaterial.Get(), MipAdjustedRenderTarget.Get());
+		bInvalidatedMipTexture = false;
+	}
+
 	return true;
 }
 
@@ -308,12 +338,14 @@ void FMediaImageViewer::PaintPanel(FMediaImageSlatePaintParams& InPaintParams, c
 void FMediaImageViewer::PaintImage(FMediaImageSlatePaintParams& InPaintParams, const FMediaImageSlatePaintGeometry& InPaintGeometry)
 {
 	TSharedPtr<FSlateBrush> BrushToUse = Brush;
+	ESlateDrawEffect DrawEffectsToUse = DrawEffects;
 
 	if (SupportsMip() && PaintSettings.bEnableMipOverride)
 	{
 		if (UpdateMipBrush())
 		{
 			BrushToUse = MipAdjustedBrush;
+			DrawEffectsToUse = ESlateDrawEffect::NoPixelSnapping | ESlateDrawEffect::NoGamma;
 		}
 	}
 
@@ -346,7 +378,7 @@ void FMediaImageViewer::PaintImage(FMediaImageSlatePaintParams& InPaintParams, c
 		InPaintParams.LayerId,
 		InPaintGeometry.Geometry,
 		BrushToUse.Get(),
-		DrawEffects,
+		DrawEffectsToUse,
 		GetPaintSettings().Rotation.Yaw * 2.f * PI / 360.f,
 		{},
 		FSlateDrawElement::ERotationSpace::RelativeToElement,
diff --git a/Engine/Plugins/Media/MediaViewer/Source/MediaViewer/Private/ImageViewers/MediaSourceImageViewer.cpp b/Engine/Plugins/Media/MediaViewer/Source/MediaViewer/Private/ImageViewers/MediaSourceImageViewer.cpp
--- a/Engine/Plugins/Media/MediaViewer/Source/MediaViewer/Private/ImageViewers/MediaSourceImageViewer.cpp
+++ b/Engine/Plugins/Media/MediaViewer/Source/MediaViewer/Private/ImageViewers/MediaSourceImageViewer.cpp
@@ -591,6 +591,9 @@ void FMediaSourceImageViewer::PaintImage(FMediaImageSlatePaintParams& InPaintPar
 		}
 	}
 
+	// Always invalidate the mip texture for video.
+	bInvalidatedMipTexture = true;
+
 	FMediaImageViewer::PaintImage(InPaintParams, InPaintGeometry);
 }
 
diff --git a/Engine/Plugins/Media/MediaViewer/Source/MediaViewer/Private/ImageViewers/ViewportImageViewer.cpp b/Engine/Plugins/Media/MediaViewer/Source/MediaViewer/Private/ImageViewers/ViewportImageViewer.cpp
--- a/Engine/Plugins/Media/MediaViewer/Source/MediaViewer/Private/ImageViewers/ViewportImageViewer.cpp
+++ b/Engine/Plugins/Media/MediaViewer/Source/MediaViewer/Private/ImageViewers/ViewportImageViewer.cpp
@@ -222,6 +222,7 @@ void FViewportImageViewer::RenderViewport(TNotNull<FViewport*> InViewport, TNotN
 
 void FViewportImageViewer::RequestRender(bool bInResizeTargetToViewport)
 {
+	bInvalidatedMipTexture = true;
 	RenderState = ERenderState::Invalidated;
 	bResizeTargetToViewport = bResizeTargetToViewport || bInResizeTargetToViewport;
 
diff --git a/Engine/Plugins/Media/MediaViewer/Source/MediaViewer/Private/ImageViewers/ViewportImageViewer.h b/Engine/Plugins/Media/MediaViewer/Source/MediaViewer/Private/ImageViewers/ViewportImageViewer.h
--- a/Engine/Plugins/Media/MediaViewer/Source/MediaViewer/Private/ImageViewers/ViewportImageViewer.h
+++ b/Engine/Plugins/Media/MediaViewer/Source/MediaViewer/Private/ImageViewers/ViewportImageViewer.h
@@ -24,7 +24,7 @@ struct FViewportImageViewerSettings
 	UPROPERTY(VisibleAnywhere, Category = "Material")
 	TObjectPtr<UTextureRenderTarget2D> RenderTarget = nullptr;
 
-	/** Will render the material every frame. */
+	/** Will render the viewport every frame. */
 	UPROPERTY(EditAnywhere, Category = "Material")
 	bool bRealTime = true;
 };
diff --git a/Engine/Plugins/Media/MediaViewer/Source/MediaViewer/Public/ImageViewer/MediaImageViewer.h b/Engine/Plugins/Media/MediaViewer/Source/MediaViewer/Public/ImageViewer/MediaImageViewer.h
--- a/Engine/Plugins/Media/MediaViewer/Source/MediaViewer/Public/ImageViewer/MediaImageViewer.h
+++ b/Engine/Plugins/Media/MediaViewer/Source/MediaViewer/Public/ImageViewer/MediaImageViewer.h
@@ -392,6 +392,7 @@ protected:
 	TStrongObjectPtr<UMaterialInstanceDynamic> MipAdjustedMaterial;
 	TStrongObjectPtr<UTextureRenderTarget2D> MipAdjustedRenderTarget;
 	TSharedPtr<FSlateBrush> MipAdjustedBrush;
+	bool bInvalidatedMipTexture = true;
 
 	bool UpdateMipBrush();
 
