Media Viewer: Fixed grammatical error in enum description.

#Synced-CL 51317559

diff --git a/Engine/Plugins/Media/MediaViewer/Source/MediaViewer/Private/Widgets/MediaViewerSettings.h b/Engine/Plugins/Media/MediaViewer/Source/MediaViewer/Private/Widgets/MediaViewerSettings.h
--- a/Engine/Plugins/Media/MediaViewer/Source/MediaViewer/Private/Widgets/MediaViewerSettings.h
+++ b/Engine/Plugins/Media/MediaViewer/Source/MediaViewer/Private/Widgets/MediaViewerSettings.h
@@ -19,7 +19,7 @@ enum class EMediaViewerMediaSyncType : uint8
 	/** The offset is relative to the start time on both videos. A scrub change of 2s in video A will move video B also by 2s. */
 	RelativeTime,
 
-	/** The offset is relative to slider value on both videos. A scrub change of 10% will in video A will move video B by 10%. */
+	/** The offset is relative to slider value on both videos. A scrub change of 10% in video A will move video B by 10%. */
 	RelativeProgress,
 
 	/** Time is synced exactly between videos. Video A changing from 5s to 10s will change video B from *anything* to 10s. */
