Media Viewer: Potential crash fix for viewport library.

#Synced-CL 51317559

diff --git a/Engine/Plugins/Media/MediaViewer/Source/MediaViewer/Private/Library/MediaViewerLibraryDynamicGroup.cpp b/Engine/Plugins/Media/MediaViewer/Source/MediaViewer/Private/Library/MediaViewerLibraryDynamicGroup.cpp
--- a/Engine/Plugins/Media/MediaViewer/Source/MediaViewer/Private/Library/MediaViewerLibraryDynamicGroup.cpp
+++ b/Engine/Plugins/Media/MediaViewer/Source/MediaViewer/Private/Library/MediaViewerLibraryDynamicGroup.cpp
@@ -66,17 +66,36 @@ TOptional<TArray<FGuid>> FMediaViewerLibraryDynamicGroup::GetUpdatedIs(const TAr
 	}
 
 	// Generate new id list
-	TArray<FGuid> UpdatedIds;
-	bool bHasNewItem = false;
-
 	const TArray<TSharedRef<FMediaViewerLibraryItem>> UpdatedItems = GenerateItemsDelegate.Execute();
 
+	TSet<FGuid> UpdatedIds;
+	UpdatedIds.Reserve(UpdatedItems.Num());
+
+	TSet<FString> StringValues;
+	StringValues.Reserve(UpdatedItems.Num());
+
+	bool bHasNewItem = false;
+
 	for (const TSharedRef<FMediaViewerLibraryItem>& UpdatedItem : UpdatedItems)
 	{
 		const FString StringValue = UpdatedItem->GetStringValue();
 
+		// Ensure no duplicates
+		if (StringValues.Contains(StringValue))
+		{
+			continue;
+		}
+
+		StringValues.Add(StringValue);
+
 		if (const FGuid* ExistingId = ExistingItems.Find(StringValue))
 		{
+			// Ensure no duplicates
+			if (UpdatedIds.Contains(*ExistingId))
+			{
+				continue;
+			}
+
 			ExistingItems.Remove(StringValue);
 			UpdatedIds.Add(*ExistingId);
 		}
@@ -100,7 +119,7 @@ TOptional<TArray<FGuid>> FMediaViewerLibraryDynamicGroup::GetUpdatedIs(const TAr
 		Library->RemoveItem(ItemPair.Value);
 	}
 
-	return UpdatedIds;
+	return UpdatedIds.Array();
 }
 
 } // UE::MediaViewer
