Media Viewer: Disable bilinear filtering when zooming in more than 200% on a given texture or video for better per-pixel inspection capability. Works with MediaTextures as well.
#Synced-CL 51835235

diff --git a/Engine/Plugins/Media/MediaViewer/Source/MediaViewer/Private/ImageViewers/MediaImageViewer.cpp b/Engine/Plugins/Media/MediaViewer/Source/MediaViewer/Private/ImageViewers/MediaImageViewer.cpp
--- a/Engine/Plugins/Media/MediaViewer/Source/MediaViewer/Private/ImageViewers/MediaImageViewer.cpp
+++ b/Engine/Plugins/Media/MediaViewer/Source/MediaViewer/Private/ImageViewers/MediaImageViewer.cpp
@@ -4,6 +4,7 @@
 
 #include "Engine/Canvas.h"
 #include "Engine/Texture.h"
+#include "MediaTexture.h"
 #include "Engine/TextureRenderTarget2D.h"
 #include "Engine/World.h"
 #include "GameFramework/Actor.h"
@@ -62,7 +63,9 @@ void FMediaImageViewer::NotifyPostChange(const FPropertyChangedEvent& InProperty
 {
 	InitBackgroundTexture();
 
-	if (InPropertyChangedEvent.GetMemberPropertyName() == GET_MEMBER_NAME_CHECKED(FMediaImagePaintSettings, MipOverride))
+	if (InPropertyChangedEvent.GetMemberPropertyName() == GET_MEMBER_NAME_CHECKED(FMediaImagePaintSettings, MipOverride)
+		|| InPropertyChangedEvent.GetMemberPropertyName() == GET_MEMBER_NAME_CHECKED(FMediaImagePaintSettings, bAutoFilteringEnabled)
+		|| InPropertyChangedEvent.GetMemberPropertyName() == GET_MEMBER_NAME_CHECKED(FMediaImagePaintSettings, AutoFilteringThreshold))
 	{
 		// This should already be invalidated by the mip level changing the required render texture size, but we'll make sure.
 		bInvalidatedMipTexture = true;
@@ -209,7 +212,11 @@ bool FMediaImageViewer::UpdateMipBrush()
 	MipAdjustedMaterial->SetTextureParameterValue(TEXT("Texture"), BrushTexture);
 	MipAdjustedMaterial->SetScalarParameterValue(TEXT("MipLevel"), static_cast<float>(PaintSettings.GetMipLevel()));
 
-	if (bInvalidatedMipTexture)
+	// MediaTextures update their content every frame, so always re-render for them.
+	// For static textures, only re-render when explicitly invalidated.
+	const bool bNeedsRender = bInvalidatedMipTexture || BrushTexture->IsA<UMediaTexture>();
+
+	if (bNeedsRender)
 	{
 		Private::FMediaViewerUtils::RenderMaterial(MipAdjustedMaterial.Get(), MipAdjustedRenderTarget.Get());
 		bInvalidatedMipTexture = false;
@@ -340,7 +347,10 @@ void FMediaImageViewer::PaintImage(FMediaImageSlatePaintParams& InPaintParams, c
 	TSharedPtr<FSlateBrush> BrushToUse = Brush;
 	ESlateDrawEffect DrawEffectsToUse = DrawEffects;
 
-	if (SupportsMip() && PaintSettings.bEnableMipOverride)
+	const bool bShouldDisableFiltering = PaintSettings.bAutoFilteringEnabled
+		&& (PaintSettings.Scale >= PaintSettings.AutoFilteringThreshold);
+
+	if (SupportsMip() && (PaintSettings.bEnableMipOverride || bShouldDisableFiltering))
 	{
 		if (UpdateMipBrush())
 		{
diff --git a/Engine/Plugins/Media/MediaViewer/Source/MediaViewer/Private/MediaViewerCommands.cpp b/Engine/Plugins/Media/MediaViewer/Source/MediaViewer/Private/MediaViewerCommands.cpp
--- a/Engine/Plugins/Media/MediaViewer/Source/MediaViewer/Private/MediaViewerCommands.cpp
+++ b/Engine/Plugins/Media/MediaViewer/Source/MediaViewer/Private/MediaViewerCommands.cpp
@@ -144,6 +144,14 @@ void FMediaViewerCommands::RegisterCommands()
 		FInputChord(EKeys::Hyphen)
 	);
 
+	UI_COMMAND(
+		ToggleAutoFiltering,
+		"Toggle Auto Filtering",
+		"Toggles automatic filtering. When enabled, bilinear filtering is disabled once the zoom level reaches the configured threshold, allowing individual pixels to be seen clearly.",
+		EUserInterfaceActionType::ToggleButton,
+		FInputChord(EKeys::I)
+	);
+
 	UI_COMMAND(
 		ToggleMipOverride,
 		"Toggle Mip Override",
diff --git a/Engine/Plugins/Media/MediaViewer/Source/MediaViewer/Private/MediaViewerCommands.h b/Engine/Plugins/Media/MediaViewer/Source/MediaViewer/Private/MediaViewerCommands.h
--- a/Engine/Plugins/Media/MediaViewer/Source/MediaViewer/Private/MediaViewerCommands.h
+++ b/Engine/Plugins/Media/MediaViewer/Source/MediaViewer/Private/MediaViewerCommands.h
@@ -35,6 +35,8 @@ public:
 	TSharedPtr<FUICommandInfo> ScaleUp;
 	TSharedPtr<FUICommandInfo> ScaleDown;
 
+	TSharedPtr<FUICommandInfo> ToggleAutoFiltering;
+
 	TSharedPtr<FUICommandInfo> ToggleMipOverride;
 	TSharedPtr<FUICommandInfo> IncreaseMipLevel;
 	TSharedPtr<FUICommandInfo> DecreaseMipLevel;
diff --git a/Engine/Plugins/Media/MediaViewer/Source/MediaViewer/Private/Widgets/SMediaImageViewerOverlay.cpp b/Engine/Plugins/Media/MediaViewer/Source/MediaViewer/Private/Widgets/SMediaImageViewerOverlay.cpp
--- a/Engine/Plugins/Media/MediaViewer/Source/MediaViewer/Private/Widgets/SMediaImageViewerOverlay.cpp
+++ b/Engine/Plugins/Media/MediaViewer/Source/MediaViewer/Private/Widgets/SMediaImageViewerOverlay.cpp
@@ -461,6 +461,13 @@ void SMediaImageViewerOverlay::BindCommands()
 	CommandList->MapAction(Commands.ScaleUp, FExecuteAction::CreateSP(this, &SMediaImageViewerOverlay::ScaleUp));
 	CommandList->MapAction(Commands.ScaleDown, FExecuteAction::CreateSP(this, &SMediaImageViewerOverlay::ScaleDown));
 
+	CommandList->MapAction(
+		Commands.ToggleAutoFiltering,
+		FExecuteAction::CreateSP(this, &SMediaImageViewerOverlay::ToggleAutoFiltering),
+		FCanExecuteAction(),
+		FGetActionCheckState::CreateSP(this, &SMediaImageViewerOverlay::IsAutoFilteringEnabled)
+	);
+
 	CommandList->MapAction(
 		Commands.ToggleMipOverride,
 		FExecuteAction::CreateSP(this, &SMediaImageViewerOverlay::ToggleMipOverride),
@@ -927,6 +934,26 @@ void SMediaImageViewerOverlay::ScaleToFit(bool bInUseTransformLock)
 	}
 }
 
+ECheckBoxState SMediaImageViewerOverlay::IsAutoFilteringEnabled() const
+{
+	if (TSharedPtr<FMediaImageViewer> ImageViewer = Delegates->GetImageViewer.Execute(Position))
+	{
+		return ImageViewer->GetPaintSettings().bAutoFilteringEnabled
+			? ECheckBoxState::Checked
+			: ECheckBoxState::Unchecked;
+	}
+
+	return ECheckBoxState::Undetermined;
+}
+
+void SMediaImageViewerOverlay::ToggleAutoFiltering()
+{
+	if (TSharedPtr<FMediaImageViewer> ImageViewer = Delegates->GetImageViewer.Execute(Position))
+	{
+		ImageViewer->GetPaintSettings().bAutoFilteringEnabled = !ImageViewer->GetPaintSettings().bAutoFilteringEnabled;
+	}
+}
+
 ECheckBoxState SMediaImageViewerOverlay::GetMipOverrideState() const
 {
 	if (TSharedPtr<FMediaImageViewer> ImageViewer = Delegates->GetImageViewer.Execute(Position))
@@ -1154,6 +1181,8 @@ TSharedRef<SWidget> SMediaImageViewerOverlay::CreateMenu()
 			MenuBuilder.AddMenuEntry(Commands.CopyTransform);
 			MenuBuilder.AddMenuEntry(Commands.AddToLibrary);
 			MenuBuilder.AddSeparator();
+			MenuBuilder.AddMenuEntry(Commands.ToggleAutoFiltering);
+			MenuBuilder.AddSeparator();
 			MenuBuilder.AddMenuEntry(Commands.ToggleMipOverride);
 			MenuBuilder.AddMenuEntry(Commands.IncreaseMipLevel);
 			MenuBuilder.AddMenuEntry(Commands.DecreaseMipLevel);
diff --git a/Engine/Plugins/Media/MediaViewer/Source/MediaViewer/Private/Widgets/SMediaImageViewerOverlay.h b/Engine/Plugins/Media/MediaViewer/Source/MediaViewer/Private/Widgets/SMediaImageViewerOverlay.h
--- a/Engine/Plugins/Media/MediaViewer/Source/MediaViewer/Private/Widgets/SMediaImageViewerOverlay.h
+++ b/Engine/Plugins/Media/MediaViewer/Source/MediaViewer/Private/Widgets/SMediaImageViewerOverlay.h
@@ -152,6 +152,9 @@ protected:
 	void Try_MultipyScale(float InMultiple);
 	void MultiplyScale(float InMultiple);
 
+	ECheckBoxState IsAutoFilteringEnabled() const;
+	void ToggleAutoFiltering();
+
 	ECheckBoxState GetMipOverrideState() const;
 	void ToggleMipOverride();
 	void IncreaseMipLevel();
diff --git a/Engine/Plugins/Media/MediaViewer/Source/MediaViewer/Private/Widgets/SMediaViewerToolbar.cpp b/Engine/Plugins/Media/MediaViewer/Source/MediaViewer/Private/Widgets/SMediaViewerToolbar.cpp
--- a/Engine/Plugins/Media/MediaViewer/Source/MediaViewer/Private/Widgets/SMediaViewerToolbar.cpp
+++ b/Engine/Plugins/Media/MediaViewer/Source/MediaViewer/Private/Widgets/SMediaViewerToolbar.cpp
@@ -390,6 +390,8 @@ TSharedRef<SWidget> SMediaViewerToolbar::MakeScaleMenu(EMediaImageViewerPosition
 		MenuBuilder.AddMenuEntry(Commands.ScaleUp);
 		MenuBuilder.AddMenuEntry(Commands.ScaleDown);
 		MenuBuilder.AddMenuEntry(Commands.ScaleToFit);
+		MenuBuilder.AddSeparator();
+		MenuBuilder.AddMenuEntry(Commands.ToggleAutoFiltering);
 	}
 
 	return MenuBuilder.MakeWidget();
diff --git a/Engine/Plugins/Media/MediaViewer/Source/MediaViewer/Public/ImageViewer/MediaImageViewer.h b/Engine/Plugins/Media/MediaViewer/Source/MediaViewer/Public/ImageViewer/MediaImageViewer.h
--- a/Engine/Plugins/Media/MediaViewer/Source/MediaViewer/Public/ImageViewer/MediaImageViewer.h
+++ b/Engine/Plugins/Media/MediaViewer/Source/MediaViewer/Public/ImageViewer/MediaImageViewer.h
@@ -114,6 +114,22 @@ struct FMediaImagePaintSettings
 	UPROPERTY(EditAnywhere, Config, Category = "Media", meta = (Bitmask, BitmaskEnum = "/Script/MediaViewer.EMediaImageMirrorPlane"))
 	uint8 MirrorPlanes = 0;
 
+	/**
+	 * When true, bilinear filtering is automatically disabled (nearest-neighbor sampling)
+	 * once the zoom level reaches AutoFilteringThreshold. This allows individual pixels to be
+	 * seen clearly when zoomed in far enough.
+	 */
+	UPROPERTY(EditAnywhere, Config, Category = "Media")
+	bool bAutoFilteringEnabled = true;
+
+	/**
+	 * The zoom scale at which filtering is automatically disabled when bAutoFilteringEnabled is true.
+	 * A value of 2.0 means filtering is disabled at 200% zoom and above.
+	 */
+	UPROPERTY(EditAnywhere, Config, Category = "Media",
+		meta = (EditCondition = "bAutoFilteringEnabled", UIMin = 1.0, ClampMin = 1.0, UIMax = 8.0, ClampMax = 8.0, Units = "x"))
+	float AutoFilteringThreshold = 2.0f;
+
 	/** Whether to enable the MipOverride */
 	UPROPERTY(EditAnywhere, Config, Category = "Media", meta = (InlineEditConditionToggle))
 	bool bEnableMipOverride = false;
