[MediaViewer] Clean up to avoid incorrect specialization in the base class.

#Synced-CL 51094493

diff --git a/Engine/Plugins/Media/MediaViewer/Source/MediaViewer/Private/ImageViewers/MediaImageViewer.cpp b/Engine/Plugins/Media/MediaViewer/Source/MediaViewer/Private/ImageViewers/MediaImageViewer.cpp
--- a/Engine/Plugins/Media/MediaViewer/Source/MediaViewer/Private/ImageViewers/MediaImageViewer.cpp
+++ b/Engine/Plugins/Media/MediaViewer/Source/MediaViewer/Private/ImageViewers/MediaImageViewer.cpp
@@ -4,7 +4,6 @@
 
 #include "Engine/Canvas.h"
 #include "Engine/Texture.h"
-#include "MediaTexture.h"
 #include "Engine/TextureRenderTarget2D.h"
 #include "Engine/World.h"
 #include "GameFramework/Actor.h"
@@ -212,11 +211,7 @@ bool FMediaImageViewer::UpdateMipBrush()
 	MipAdjustedMaterial->SetTextureParameterValue(TEXT("Texture"), BrushTexture);
 	MipAdjustedMaterial->SetScalarParameterValue(TEXT("MipLevel"), static_cast<float>(PaintSettings.GetMipLevel()));
 
-	// MediaTextures update their content every frame, so always re-render for them.
-	// For static textures, only re-render when explicitly invalidated.
-	const bool bNeedsRender = bInvalidatedMipTexture || BrushTexture->IsA<UMediaTexture>();
-
-	if (bNeedsRender)
+	if (bInvalidatedMipTexture)
 	{
 		Private::FMediaViewerUtils::RenderMaterial(MipAdjustedMaterial.Get(), MipAdjustedRenderTarget.Get());
 		bInvalidatedMipTexture = false;
diff --git a/Engine/Plugins/Media/MediaViewer/Source/MediaViewer/Private/ImageViewers/MediaTextureImageViewer.cpp b/Engine/Plugins/Media/MediaViewer/Source/MediaViewer/Private/ImageViewers/MediaTextureImageViewer.cpp
--- a/Engine/Plugins/Media/MediaViewer/Source/MediaViewer/Private/ImageViewers/MediaTextureImageViewer.cpp
+++ b/Engine/Plugins/Media/MediaViewer/Source/MediaViewer/Private/ImageViewers/MediaTextureImageViewer.cpp
@@ -336,6 +336,14 @@ FString FMediaTextureImageViewer::GetReferencerName() const
 	return ReferencerName;
 }
 
+void FMediaTextureImageViewer::PaintImage(FMediaImageSlatePaintParams& InPaintParams, const FMediaImageSlatePaintGeometry& InPaintGeometry)
+{
+	// Always invalidate the mip texture for video.
+	bInvalidatedMipTexture = true;
+
+	FMediaImageViewer::PaintImage(InPaintParams, InPaintGeometry);
+}
+	
 } // UE::MediaViewer::Private
 
 #undef LOCTEXT_NAMESPACE
diff --git a/Engine/Plugins/Media/MediaViewer/Source/MediaViewer/Private/ImageViewers/MediaTextureImageViewer.h b/Engine/Plugins/Media/MediaViewer/Source/MediaViewer/Private/ImageViewers/MediaTextureImageViewer.h
--- a/Engine/Plugins/Media/MediaViewer/Source/MediaViewer/Private/ImageViewers/MediaTextureImageViewer.h
+++ b/Engine/Plugins/Media/MediaViewer/Source/MediaViewer/Private/ImageViewers/MediaTextureImageViewer.h
@@ -96,6 +96,10 @@ public:
 	//~ End FNotifyHook
 
 protected:
+	//~ Begin FMediaImageViewer
+	virtual void PaintImage(FMediaImageSlatePaintParams& InPaintParams, const FMediaImageSlatePaintGeometry& InPaintGeometry) override;
+	//~ End FMediaImageViewer
+
 	FMediaTextureImageViewerSettings MediaTextureSettings;
 	TSharedPtr<FTextureSampleCache> SampleCache;
 };
