[Media Viewer] Display 2D grid for better pixel visibility when zoomed in.

Pixel Grid Overlay
+ Draws a grid over the image when zoomed in with nearest-neighbor filtering active, making individual pixel boundaries visible
+ Grid activates when a source pixel occupies = 9 screen pixels (CellSize >= PixelGridThreshold, default 9.f)
+ Grid cell size accounts for the active mip level (Scale * 2^MipLevel), so lines always align to actual rendered pixels regardless of mip override
"Display Pixel Grid" toggle command added (key G), appears in both the overlay context menu and toolbar scale dropdown
+ Toggle is grayed out when Auto Filtering is disabled
+ bPixelGridEnabled and PixelGridThreshold added to FMediaImagePaintSettings

Zoom Level Overhaul
+ Extended max zoom from 1000% to 2000%
+ MaxScale bumped from 10.f to 20.f

Scale To Fit Fix
Pan/translation offset is now reset before the fit scale is applied, fixing cases where Scale to Fit would fail after the image had been translated

#Synced-CL 52257653

diff --git a/Engine/Plugins/Media/MediaViewer/Source/MediaViewer/Private/ImageViewers/MediaImageViewer.cpp b/Engine/Plugins/Media/MediaViewer/Source/MediaViewer/Private/ImageViewers/MediaImageViewer.cpp
--- a/Engine/Plugins/Media/MediaViewer/Source/MediaViewer/Private/ImageViewers/MediaImageViewer.cpp
+++ b/Engine/Plugins/Media/MediaViewer/Source/MediaViewer/Private/ImageViewers/MediaImageViewer.cpp
@@ -300,6 +300,7 @@ void FMediaImageViewer::Paint(FMediaImageSlatePaintParams& InPaintParams)
 
 	PaintPanel(InPaintParams, Geometry);
 	PaintImage(InPaintParams, Geometry);
+	PaintPixelGrid(InPaintParams, Geometry);
 
 	InPaintParams.DrawElements.PopClip();
 }
@@ -393,6 +394,72 @@ void FMediaImageViewer::PaintImage(FMediaImageSlatePaintParams& InPaintParams, c
 	++InPaintParams.LayerId;
 }
 
+void FMediaImageViewer::PaintPixelGrid(FMediaImageSlatePaintParams& InPaintParams, const FMediaImageSlatePaintGeometry& InPaintGeometry)
+{
+	// At mip level N each displayed "pixel" covers 2^N source pixels, so widen the
+	// grid cell size and shrink the iteration bounds to match the mip resolution.
+	const int32 MipLevel  = static_cast<int32>(PaintSettings.GetMipLevel());
+	const float CellSize  = PaintSettings.Scale * static_cast<float>(1 << MipLevel);
+
+	// Show the grid once a single source pixel occupies at least the configured screen-pixel threshold.
+	const bool bShouldShowGrid = PaintSettings.bPixelGridEnabled
+		&& PaintSettings.bAutoFilteringEnabled
+		&& (CellSize >= PaintSettings.PixelGridThreshold);
+
+	if (!bShouldShowGrid)
+	{
+		return;
+	}
+	const int32 MipWidth  = FMath::Max(1, ImageInfo.Size.X >> MipLevel);
+	const int32 MipHeight = FMath::Max(1, ImageInfo.Size.Y >> MipLevel);
+
+	const FVector2D& Position = InPaintGeometry.Position;
+
+	// Intersect image bounds with viewer bounds to avoid drawing outside either
+	const float ClipLeft   = FMath::Max(Position.X, InPaintParams.ViewerPosition.X);
+	const float ClipTop    = FMath::Max(Position.Y, InPaintParams.ViewerPosition.Y);
+	const float ClipRight  = FMath::Min(Position.X + InPaintGeometry.Size.X, InPaintParams.ViewerPosition.X + InPaintParams.ViewerSize.X);
+	const float ClipBottom = FMath::Min(Position.Y + InPaintGeometry.Size.Y, InPaintParams.ViewerPosition.Y + InPaintParams.ViewerSize.Y);
+
+	if (ClipLeft >= ClipRight || ClipTop >= ClipBottom)
+	{
+		return;
+	}
+
+	// Determine which pixel boundary lines are visible
+	const int32 FirstCol = FMath::Max(0, FMath::FloorToInt((ClipLeft  - Position.X) / CellSize));
+	const int32 LastCol  = FMath::Min(MipWidth,  FMath::CeilToInt((ClipRight  - Position.X) / CellSize));
+	const int32 FirstRow = FMath::Max(0, FMath::FloorToInt((ClipTop   - Position.Y) / CellSize));
+	const int32 LastRow  = FMath::Min(MipHeight, FMath::CeilToInt((ClipBottom - Position.Y) / CellSize));
+
+	const FLinearColor GridColor(0.f, 0.f, 0.f, 0.4f);
+	const FPaintGeometry FullGeometry = InPaintParams.AllottedGeometry.ToPaintGeometry();
+	constexpr ESlateDrawEffect DrawEffect = ESlateDrawEffect::NoPixelSnapping;
+
+	TArray<FVector2D> LinePoints;
+	LinePoints.SetNumUninitialized(2);
+
+	// Vertical lines (pixel column boundaries)
+	for (int32 Col = FirstCol; Col <= LastCol; ++Col)
+	{
+		const float X = Position.X + Col * CellSize;
+		LinePoints[0] = FVector2D(X, ClipTop);
+		LinePoints[1] = FVector2D(X, ClipBottom);
+		FSlateDrawElement::MakeLines(InPaintParams.DrawElements, InPaintParams.LayerId, FullGeometry, LinePoints, DrawEffect, GridColor, false, 1.f);
+	}
+
+	// Horizontal lines (pixel row boundaries)
+	for (int32 Row = FirstRow; Row <= LastRow; ++Row)
+	{
+		const float Y = Position.Y + Row * CellSize;
+		LinePoints[0] = FVector2D(ClipLeft, Y);
+		LinePoints[1] = FVector2D(ClipRight, Y);
+		FSlateDrawElement::MakeLines(InPaintParams.DrawElements, InPaintParams.LayerId, FullGeometry, LinePoints, DrawEffect, GridColor, false, 1.f);
+	}
+
+	++InPaintParams.LayerId;
+}
+
 void FMediaImageViewer::Paint(const FMediaImageCanvasPaintParams& InPaintParams)
 {
 	const FVector2D CanvasSize = {
diff --git a/Engine/Plugins/Media/MediaViewer/Source/MediaViewer/Private/MediaViewerCommands.cpp b/Engine/Plugins/Media/MediaViewer/Source/MediaViewer/Private/MediaViewerCommands.cpp
--- a/Engine/Plugins/Media/MediaViewer/Source/MediaViewer/Private/MediaViewerCommands.cpp
+++ b/Engine/Plugins/Media/MediaViewer/Source/MediaViewer/Private/MediaViewerCommands.cpp
@@ -152,6 +152,14 @@ void FMediaViewerCommands::RegisterCommands()
 		FInputChord(EKeys::I)
 	);
 
+	UI_COMMAND(
+		TogglePixelGrid,
+		"Display Pixel Grid",
+		"Toggles the pixel grid overlay. When enabled, a grid is drawn over the image at high zoom levels to make individual pixel boundaries clearly visible.",
+		EUserInterfaceActionType::ToggleButton,
+		FInputChord(EKeys::G)
+	);
+
 	UI_COMMAND(
 		ToggleMipOverride,
 		"Toggle Mip Override",
@@ -240,6 +248,14 @@ void FMediaViewerCommands::RegisterCommands()
 		FInputChord()
 	);
 
+	UI_COMMAND(
+		Scale1000,
+		"1000%",
+		"Set Scale to 1000%.",
+		EUserInterfaceActionType::RadioButton,
+		FInputChord()
+	);
+
 	UI_COMMAND(
 		CopyTransform,
 		"Copy Transform",
diff --git a/Engine/Plugins/Media/MediaViewer/Source/MediaViewer/Private/MediaViewerCommands.h b/Engine/Plugins/Media/MediaViewer/Source/MediaViewer/Private/MediaViewerCommands.h
--- a/Engine/Plugins/Media/MediaViewer/Source/MediaViewer/Private/MediaViewerCommands.h
+++ b/Engine/Plugins/Media/MediaViewer/Source/MediaViewer/Private/MediaViewerCommands.h
@@ -36,6 +36,7 @@ public:
 	TSharedPtr<FUICommandInfo> ScaleDown;
 
 	TSharedPtr<FUICommandInfo> ToggleAutoFiltering;
+	TSharedPtr<FUICommandInfo> TogglePixelGrid;
 
 	TSharedPtr<FUICommandInfo> ToggleMipOverride;
 	TSharedPtr<FUICommandInfo> IncreaseMipLevel;
@@ -48,6 +49,7 @@ public:
 	TSharedPtr<FUICommandInfo> Scale200;
 	TSharedPtr<FUICommandInfo> Scale400;
 	TSharedPtr<FUICommandInfo> Scale800;
+	TSharedPtr<FUICommandInfo> Scale1000;
 	TSharedPtr<FUICommandInfo> ScaleToFit;
 
 	TSharedPtr<FUICommandInfo> CopyTransform;
diff --git a/Engine/Plugins/Media/MediaViewer/Source/MediaViewer/Private/Widgets/SMediaImageViewerOverlay.cpp b/Engine/Plugins/Media/MediaViewer/Source/MediaViewer/Private/Widgets/SMediaImageViewerOverlay.cpp
--- a/Engine/Plugins/Media/MediaViewer/Source/MediaViewer/Private/Widgets/SMediaImageViewerOverlay.cpp
+++ b/Engine/Plugins/Media/MediaViewer/Source/MediaViewer/Private/Widgets/SMediaImageViewerOverlay.cpp
@@ -468,6 +468,13 @@ void SMediaImageViewerOverlay::BindCommands()
 		FGetActionCheckState::CreateSP(this, &SMediaImageViewerOverlay::IsAutoFilteringEnabled)
 	);
 
+	CommandList->MapAction(
+		Commands.TogglePixelGrid,
+		FExecuteAction::CreateSP(this, &SMediaImageViewerOverlay::TogglePixelGrid),
+		FCanExecuteAction::CreateSP(this, &SMediaImageViewerOverlay::CanTogglePixelGrid),
+		FGetActionCheckState::CreateSP(this, &SMediaImageViewerOverlay::IsPixelGridEnabled)
+	);
+
 	CommandList->MapAction(
 		Commands.ToggleMipOverride,
 		FExecuteAction::CreateSP(this, &SMediaImageViewerOverlay::ToggleMipOverride),
@@ -489,7 +496,8 @@ void SMediaImageViewerOverlay::BindCommands()
 	CommandList->MapAction(Commands.Scale100, FExecuteAction::CreateSP(this, &SMediaImageViewerOverlay::Try_SetScale, 1.0f));
 	CommandList->MapAction(Commands.Scale200, FExecuteAction::CreateSP(this, &SMediaImageViewerOverlay::Try_SetScale, 2.0f));
 	CommandList->MapAction(Commands.Scale400, FExecuteAction::CreateSP(this, &SMediaImageViewerOverlay::Try_SetScale, 4.0f));
-	CommandList->MapAction(Commands.Scale800, FExecuteAction::CreateSP(this, &SMediaImageViewerOverlay::Try_SetScale, 8.0f));
+	CommandList->MapAction(Commands.Scale800,  FExecuteAction::CreateSP(this, &SMediaImageViewerOverlay::Try_SetScale,  8.0f));
+	CommandList->MapAction(Commands.Scale1000, FExecuteAction::CreateSP(this, &SMediaImageViewerOverlay::Try_SetScale, 10.0f));
 
 	CommandList->MapAction(
 		Commands.ResetTransform, 
@@ -860,15 +868,11 @@ float SMediaImageViewerOverlay::GetScaleModifier()
 
 void SMediaImageViewerOverlay::ScaleUp()
 {
-	const float ModifierScale = GetScaleModifier();
-
 	Try_MultipyScale(GetScaleModifier());
 }
 
 void SMediaImageViewerOverlay::ScaleDown()
 {
-	const float ModifierScale = GetScaleModifier();
-
 	Try_MultipyScale(1.f / GetScaleModifier());
 }
 
@@ -898,6 +902,8 @@ void SMediaImageViewerOverlay::ScaleToFit(bool bInUseTransformLock)
 		return;
 	}
 
+	SetOffset(FVector::ZeroVector);
+
 	const FVector2D LocalSize = Delegates->GetViewerSize.Execute();
 
 	if (FMath::IsNearlyZero(LocalSize.X) || FMath::IsNearlyZero(LocalSize.Y))
@@ -954,6 +960,36 @@ void SMediaImageViewerOverlay::ToggleAutoFiltering()
 	}
 }
 
+bool SMediaImageViewerOverlay::CanTogglePixelGrid() const
+{
+	if (TSharedPtr<FMediaImageViewer> ImageViewer = Delegates->GetImageViewer.Execute(Position))
+	{
+		return ImageViewer->GetPaintSettings().bAutoFilteringEnabled;
+	}
+
+	return false;
+}
+
+ECheckBoxState SMediaImageViewerOverlay::IsPixelGridEnabled() const
+{
+	if (TSharedPtr<FMediaImageViewer> ImageViewer = Delegates->GetImageViewer.Execute(Position))
+	{
+		return ImageViewer->GetPaintSettings().bPixelGridEnabled
+			? ECheckBoxState::Checked
+			: ECheckBoxState::Unchecked;
+	}
+
+	return ECheckBoxState::Undetermined;
+}
+
+void SMediaImageViewerOverlay::TogglePixelGrid()
+{
+	if (TSharedPtr<FMediaImageViewer> ImageViewer = Delegates->GetImageViewer.Execute(Position))
+	{
+		ImageViewer->GetPaintSettings().bPixelGridEnabled = !ImageViewer->GetPaintSettings().bPixelGridEnabled;
+	}
+}
+
 ECheckBoxState SMediaImageViewerOverlay::GetMipOverrideState() const
 {
 	if (TSharedPtr<FMediaImageViewer> ImageViewer = Delegates->GetImageViewer.Execute(Position))
@@ -1182,6 +1218,7 @@ TSharedRef<SWidget> SMediaImageViewerOverlay::CreateMenu()
 			MenuBuilder.AddMenuEntry(Commands.AddToLibrary);
 			MenuBuilder.AddSeparator();
 			MenuBuilder.AddMenuEntry(Commands.ToggleAutoFiltering);
+			MenuBuilder.AddMenuEntry(Commands.TogglePixelGrid);
 			MenuBuilder.AddSeparator();
 			MenuBuilder.AddMenuEntry(Commands.ToggleMipOverride);
 			MenuBuilder.AddMenuEntry(Commands.IncreaseMipLevel);
diff --git a/Engine/Plugins/Media/MediaViewer/Source/MediaViewer/Private/Widgets/SMediaImageViewerOverlay.h b/Engine/Plugins/Media/MediaViewer/Source/MediaViewer/Private/Widgets/SMediaImageViewerOverlay.h
--- a/Engine/Plugins/Media/MediaViewer/Source/MediaViewer/Private/Widgets/SMediaImageViewerOverlay.h
+++ b/Engine/Plugins/Media/MediaViewer/Source/MediaViewer/Private/Widgets/SMediaImageViewerOverlay.h
@@ -155,6 +155,10 @@ protected:
 	ECheckBoxState IsAutoFilteringEnabled() const;
 	void ToggleAutoFiltering();
 
+	bool CanTogglePixelGrid() const;
+	ECheckBoxState IsPixelGridEnabled() const;
+	void TogglePixelGrid();
+
 	ECheckBoxState GetMipOverrideState() const;
 	void ToggleMipOverride();
 	void IncreaseMipLevel();
diff --git a/Engine/Plugins/Media/MediaViewer/Source/MediaViewer/Private/Widgets/SMediaViewer.h b/Engine/Plugins/Media/MediaViewer/Source/MediaViewer/Private/Widgets/SMediaViewer.h
--- a/Engine/Plugins/Media/MediaViewer/Source/MediaViewer/Private/Widgets/SMediaViewer.h
+++ b/Engine/Plugins/Media/MediaViewer/Source/MediaViewer/Private/Widgets/SMediaViewer.h
@@ -60,7 +60,7 @@ class SMediaViewer : public SCompoundWidget, public FGCObject
 	SLATE_END_ARGS()
 
 public:
-	static constexpr float MaxScale = 10.f;
+	static constexpr float MaxScale = 20.f;
 	static constexpr float MinScale = 0.01f;
 
 	SMediaViewer();
diff --git a/Engine/Plugins/Media/MediaViewer/Source/MediaViewer/Private/Widgets/SMediaViewerToolbar.cpp b/Engine/Plugins/Media/MediaViewer/Source/MediaViewer/Private/Widgets/SMediaViewerToolbar.cpp
--- a/Engine/Plugins/Media/MediaViewer/Source/MediaViewer/Private/Widgets/SMediaViewerToolbar.cpp
+++ b/Engine/Plugins/Media/MediaViewer/Source/MediaViewer/Private/Widgets/SMediaViewerToolbar.cpp
@@ -386,12 +386,14 @@ TSharedRef<SWidget> SMediaViewerToolbar::MakeScaleMenu(EMediaImageViewerPosition
 		MenuBuilder.AddMenuEntry(Commands.Scale200);
 		MenuBuilder.AddMenuEntry(Commands.Scale400);
 		MenuBuilder.AddMenuEntry(Commands.Scale800);
+		MenuBuilder.AddMenuEntry(Commands.Scale1000);
 		MenuBuilder.AddSeparator();
 		MenuBuilder.AddMenuEntry(Commands.ScaleUp);
 		MenuBuilder.AddMenuEntry(Commands.ScaleDown);
 		MenuBuilder.AddMenuEntry(Commands.ScaleToFit);
 		MenuBuilder.AddSeparator();
 		MenuBuilder.AddMenuEntry(Commands.ToggleAutoFiltering);
+		MenuBuilder.AddMenuEntry(Commands.TogglePixelGrid);
 	}
 
 	return MenuBuilder.MakeWidget();
diff --git a/Engine/Plugins/Media/MediaViewer/Source/MediaViewer/Public/ImageViewer/MediaImageViewer.h b/Engine/Plugins/Media/MediaViewer/Source/MediaViewer/Public/ImageViewer/MediaImageViewer.h
--- a/Engine/Plugins/Media/MediaViewer/Source/MediaViewer/Public/ImageViewer/MediaImageViewer.h
+++ b/Engine/Plugins/Media/MediaViewer/Source/MediaViewer/Public/ImageViewer/MediaImageViewer.h
@@ -128,7 +128,21 @@ struct FMediaImagePaintSettings
 	 */
 	UPROPERTY(EditAnywhere, Config, Category = "Media",
 		meta = (EditCondition = "bAutoFilteringEnabled", UIMin = 1.0, ClampMin = 1.0, UIMax = 8.0, ClampMax = 8.0, Units = "x"))
-	float AutoFilteringThreshold = 2.0f;
+	float AutoFilteringThreshold = 4.0f;
+
+	/**
+	 * When true, a pixel grid overlay is drawn when zoomed in with nearest-neighbor filtering active,
+	 * making individual pixel boundaries clearly visible.
+	 */
+	UPROPERTY(EditAnywhere, Config, Category = "Media")
+	bool bPixelGridEnabled = true;
+
+	/**
+	 * Minimum screen-pixel size a source pixel must reach before the grid is drawn.
+	 */
+	UPROPERTY(EditAnywhere, Config, Category = "Media",
+		meta = (EditCondition = "bPixelGridEnabled", UIMin = 8.0, ClampMin = 1.0, UIMax = 128.0, ClampMax = 128.0))
+	float PixelGridThreshold = 9.f;
 
 	/** Whether to enable the MipOverride */
 	UPROPERTY(EditAnywhere, Config, Category = "Media", meta = (InlineEditConditionToggle))
@@ -419,6 +433,9 @@ protected:
 
 	MEDIAVIEWER_API virtual void PaintImage(FMediaImageSlatePaintParams& InPaintParams, const FMediaImageSlatePaintGeometry& InPaintGeometry);
 
+	/** Draws a pixel boundary grid over the image when nearest-neighbor filtering is active. */
+	MEDIAVIEWER_API virtual void PaintPixelGrid(FMediaImageSlatePaintParams& InPaintParams, const FMediaImageSlatePaintGeometry& InPaintGeometry);
+
 	MEDIAVIEWER_API virtual void PaintPanel(const FMediaImageCanvasPaintParams& InPaintParams, const FMediaImageCanvasPaintGeometry& InPaintGeometry);
 
 	MEDIAVIEWER_API virtual void PaintImage(const FMediaImageCanvasPaintParams& InPaintParams, const FMediaImageCanvasPaintGeometry& InPaintGeometry);
